Initial implementation of an OpenTrackIO Live Link plugin.  Supports reading OpenTrackIO data from UDP socket using the 1.0 specification found at https://www.opentrackio.org/.  It will sucessfully handle OpenTrackIO headers from either unicast endpoint or multi-cast group and parse either CBOR or JSON payloads.  The camera provided in the data stream will be reflected as a Live Link subject using the Live Link Lens role which inherits from the LiveLinkCameraRole type.    Multiple input streams is supported so many cameras can send OpenTrackIO data to the engine.

Unit tests are included to verify proper parsing of OpenTrackIO JSON / CBOR payloads with simulated packet parsing.

There are still many things missing from the implementation:
1. Not all values provided by OpenTrackIO are reflected into the lens live link role.
2. Blueprint access to the OpenTrackIO data type has not been completed.
3. Custom data is not supported yet
4. Calculation of transformation matrix has not been validated and probably missing an axis conversion.
5. Full validation of data applied to Live Link Lens role has not been done.
6. No calculation for location of pinhole camera.

diff --git a/Engine/Plugins/VirtualProduction/LiveLinkOpenTrackIO/LiveLinkOpenTrackIO.uplugin b/Engine/Plugins/VirtualProduction/LiveLinkOpenTrackIO/LiveLinkOpenTrackIO.uplugin
new file mode 100644
--- /dev/null
+++ b/Engine/Plugins/VirtualProduction/LiveLinkOpenTrackIO/LiveLinkOpenTrackIO.uplugin
@@ -0,0 +1,38 @@
+{
+	"FileVersion": 1,
+	"Version": 1,
+	"VersionName": "1.0",
+	"FriendlyName": "LiveLinkOpenTrackIO",
+	"Description": "Live Link plugin for supporting OpenTrackIO (https://opentrackio.org) devices in Unreal Engine or Live Link Hub.",
+	"Category": "Virtual Production",
+	"CreatedBy": "Epic Games, Inc.",
+	"CreatedByURL": "https://epicgames.com",
+	"DocsURL": "",
+	"MarketplaceURL": "",
+	"SupportURL": "",
+	"CanContainContent": true,
+	"IsBetaVersion": true,
+	"IsExperimentalVersion": false,
+	"Installed": false,
+	"EnabledByDefault": false,
+	"SupportedPrograms": ["LiveLinkHub"],
+	"Modules": [
+		{
+			"Name": "LiveLinkOpenTrackIO",
+			"Type": "Runtime",
+			"LoadingPhase": "Default",
+			"ProgramAllowList": ["LiveLinkHub"]
+		}
+	],
+	"Plugins": [
+		{
+			"Name": "LiveLinkLens",
+			"Enabled": true
+		},
+		{
+			"Name": "LiveLink",
+			"Enabled": true
+		}
+	]
+
+}
diff --git a/Engine/Plugins/VirtualProduction/LiveLinkOpenTrackIO/Source/LiveLinkOpenTrackIO/LiveLinkOpenTrackIO.Build.cs b/Engine/Plugins/VirtualProduction/LiveLinkOpenTrackIO/Source/LiveLinkOpenTrackIO/LiveLinkOpenTrackIO.Build.cs
new file mode 100644
--- /dev/null
+++ b/Engine/Plugins/VirtualProduction/LiveLinkOpenTrackIO/Source/LiveLinkOpenTrackIO/LiveLinkOpenTrackIO.Build.cs
@@ -0,0 +1,39 @@
+// Copyright Epic Games, Inc. All Rights Reserved.
+
+using UnrealBuildTool;
+using System.IO;
+
+public class LiveLinkOpenTrackIO : ModuleRules
+{
+	public LiveLinkOpenTrackIO(ReadOnlyTargetRules Target) : base(Target)
+	{
+		PCHUsage = ModuleRules.PCHUsageMode.UseExplicitOrSharedPCHs;
+		
+		PublicDependencyModuleNames.AddRange(
+			new string[]
+			{
+				"Core",
+				"LiveLinkInterface",
+			}
+			);
+			
+		
+		PrivateDependencyModuleNames.AddRange(
+			new string[]
+			{
+				"ApplicationCore",
+				"CoreUObject",
+				"Engine",
+				"Json",
+				"JsonUtilities",
+				"LiveLinkLens",
+				"Projects",
+				"Networking",
+				"Sockets",
+				"Slate",
+				"SlateCore",
+				"Serialization"
+			}
+			);
+	}
+}
diff --git a/Engine/Plugins/VirtualProduction/LiveLinkOpenTrackIO/Source/LiveLinkOpenTrackIO/Private/LiveLinkOpenTrackIO.cpp b/Engine/Plugins/VirtualProduction/LiveLinkOpenTrackIO/Source/LiveLinkOpenTrackIO/Private/LiveLinkOpenTrackIO.cpp
new file mode 100644
--- /dev/null
+++ b/Engine/Plugins/VirtualProduction/LiveLinkOpenTrackIO/Source/LiveLinkOpenTrackIO/Private/LiveLinkOpenTrackIO.cpp
@@ -0,0 +1,19 @@
+// Copyright Epic Games, Inc. All Rights Reserved.
+
+#include "LiveLinkOpenTrackIO.h"
+
+#define LOCTEXT_NAMESPACE "LiveLinkOpenTrackIOModule"
+
+DEFINE_LOG_CATEGORY(LogLiveLinkOpenTrackIO);
+
+void FLiveLinkOpenTrackIOModule::StartupModule()
+{
+}
+
+void FLiveLinkOpenTrackIOModule::ShutdownModule()
+{
+}
+
+#undef LOCTEXT_NAMESPACE
+	
+IMPLEMENT_MODULE(FLiveLinkOpenTrackIOModule, LiveLinkOpenTrackIO)
\ No newline at end of file
diff --git a/Engine/Plugins/VirtualProduction/LiveLinkOpenTrackIO/Source/LiveLinkOpenTrackIO/Private/LiveLinkOpenTrackIOConnectionSettings.h b/Engine/Plugins/VirtualProduction/LiveLinkOpenTrackIO/Source/LiveLinkOpenTrackIO/Private/LiveLinkOpenTrackIOConnectionSettings.h
new file mode 100644
--- /dev/null
+++ b/Engine/Plugins/VirtualProduction/LiveLinkOpenTrackIO/Source/LiveLinkOpenTrackIO/Private/LiveLinkOpenTrackIOConnectionSettings.h
@@ -0,0 +1,12 @@
+// Copyright Epic Games, Inc. All Rights Reserved.
+
+#pragma once
+#include "LiveLinkSourceSettings.h"
+
+#include "LiveLinkOpenTrackIOConnectionSettings.generated.h"
+
+USTRUCT()
+struct LIVELINKOPENTRACKIO_API FLiveLinkOpenTrackIOConnectionSettings
+{
+	GENERATED_BODY()
+};
diff --git a/Engine/Plugins/VirtualProduction/LiveLinkOpenTrackIO/Source/LiveLinkOpenTrackIO/Private/LiveLinkOpenTrackIODatagram.h b/Engine/Plugins/VirtualProduction/LiveLinkOpenTrackIO/Source/LiveLinkOpenTrackIO/Private/LiveLinkOpenTrackIODatagram.h
new file mode 100644
--- /dev/null
+++ b/Engine/Plugins/VirtualProduction/LiveLinkOpenTrackIO/Source/LiveLinkOpenTrackIO/Private/LiveLinkOpenTrackIODatagram.h
@@ -0,0 +1,79 @@
+// Copyright Epic Games, Inc. All Rights Reserved.
+#pragma once
+
+#include "UObject/ObjectMacros.h"
+
+#include "LiveLinkOpenTrackIODatagram.generated.h"
+
+/**
+ * This struct stores the header portion of a UDP or serial OpenTrackIO payload. It only contains the header portion and the payload part will be
+ * handled separately via payload parsers.
+ *
+ * See specification at https://www.opentrackio.org/
+ */
+USTRUCT()
+struct FLiveLinkOpenTrackIODatagramHeader
+{
+	GENERATED_BODY()
+
+	/**
+	 * Bit offset 0-31
+	 * 4 bytes: Static value to indicate OpenTrackIO packet, set to ASCII "OTrk" (0x4F54726B)
+	 */ 
+	UPROPERTY()
+	uint32 Identifier = 0;
+
+	/**
+	 * Bit offset 32-39
+	 * 1 byte: This field is reserved for future use and should be ignored by both Producers and Consumers.
+	 */
+	UPROPERTY()
+	uint8 Reserved = 0;
+
+	/**
+	 * Bit offset 40-47
+	 * 1 byte: Indicates the payload format (e.g., JSON = 0x01, CBOR = 0x02, OTP = 0x02). 0x80 and above are reserved for vendor specific protocols
+	 */
+	UPROPERTY()
+	uint8 Encoding = 0;
+
+	/**
+	 * Bit offset 48-63
+	 * 2 bytes: A 16-bit unsigned integer indicating the OpenTrackIO packet's unique sequence number (0x01 to UINT16)
+	 */
+	UPROPERTY()
+	uint16 SequenceNumber = 0;
+
+	/**
+	 * Bit offset 64-95
+	 * 4 bytes: A 32-bit field indicating the byte offset of this payload segment when the overall payload length necessitates segmentation. Must be set to 0x00 for single-segment payloads.
+	 */
+	UPROPERTY()
+	uint32 SegmentOffset = 0;
+
+	/**
+	 * Bit offset 96-111
+	 * The first bit shall be set to 1 if this is the only segment or the last segment in a segmented payload, or 0 if more segments are expected. The rest of the bits represent
+	 * the total lenght of the payload.
+	 *
+	 * This differs slightly from the OpenTrackIO spec in that they have broken out LastSegment and Payload as separate fields but to use as a USTRUCT we must lump them together.
+	 */
+	UPROPERTY()
+	uint16 LastSegmentFlagAndPayloadLength = 0;
+	
+	/**
+	 * Bit offset 112-127
+	 * 2 bytes: A 16-bit checksum computed using the Fletcher-16 algorithm, covering the header (excluding checksum bytes) and payload.
+	 */
+	UPROPERTY()
+	uint16 Checksum = 0;
+
+
+	/** Helper function to return the size of the payload by and-ing the last 15 bits. */
+	uint16 GetPayloadSize() const
+	{
+		return LastSegmentFlagAndPayloadLength & 0x7FFF;
+	}
+	
+};
+
diff --git a/Engine/Plugins/VirtualProduction/LiveLinkOpenTrackIO/Source/LiveLinkOpenTrackIO/Private/LiveLinkOpenTrackIOParser.cpp b/Engine/Plugins/VirtualProduction/LiveLinkOpenTrackIO/Source/LiveLinkOpenTrackIO/Private/LiveLinkOpenTrackIOParser.cpp
new file mode 100644
--- /dev/null
+++ b/Engine/Plugins/VirtualProduction/LiveLinkOpenTrackIO/Source/LiveLinkOpenTrackIO/Private/LiveLinkOpenTrackIOParser.cpp
@@ -0,0 +1,243 @@
+// Copyright Epic Games, Inc. All Rights Reserved.
+ 
+#include "LiveLinkOpenTrackIOParser.h"
+
+#include "LiveLinkOpenTrackIO.h"
+
+#include "LiveLinkOpenTrackIODatagram.h"
+#include "LiveLinkOpenTrackIOTypes.h"
+
+#include "JsonObjectConverter.h"
+
+#include "Dom/JsonObject.h"
+#include "Dom/JsonValue.h"
+#include "Serialization/JsonSerializer.h"
+#include "Serialization/MemoryReader.h"
+
+#include "Misc/ByteSwap.h"
+#include "Misc/FileHelper.h"
+
+#include "StructSerializer.h"
+#include "StructDeserializer.h"
+#include "Backends/CborStructSerializerBackend.h"
+#include "Backends/CborStructDeserializerBackend.h"
+
+namespace UE::OpenTrackIO::Private
+{
+	
+TOptional<FLiveLinkOpenTrackIOData> ParseJsonBlob(const FString& JsonBlob)
+{
+	TSharedPtr<FJsonObject> JsonObject;
+	FJsonSerializer::Deserialize(TJsonReaderFactory<TCHAR>::Create(*JsonBlob), JsonObject);
+
+	TOptional<FLiveLinkOpenTrackIOData> OutData;
+	if (!JsonObject)
+	{
+		return OutData;
+	}
+
+	TSharedPtr<FJsonValue> StaticData = JsonObject->TryGetField(TEXT("static"));
+	TSharedPtr<FJsonValue> LensData = JsonObject->TryGetField(TEXT("lens"));
+	TSharedPtr<FJsonValue> TrackerData = JsonObject->TryGetField(TEXT("timing"));
+	TSharedPtr<FJsonValue> ProtocolData = JsonObject->TryGetField(TEXT("protocol"));
+
+	/** We need to have at least one of these objects in the payload. */
+	if (!StaticData && !LensData && !TrackerData && !ProtocolData)
+	{
+		return OutData;
+	}
+	
+	FLiveLinkOpenTrackIOData Data;
+	if (FJsonObjectConverter::JsonObjectStringToUStruct(JsonBlob, &Data))
+	{
+		OutData.Emplace(MoveTemp(Data));
+	}
+	return OutData;
+}
+
+TOptional<FLiveLinkOpenTrackIOData> ParseCborBlob(TArrayView<uint8> InBytes)
+{
+	TOptional<FLiveLinkOpenTrackIOData> OutData;
+	if (InBytes.IsEmpty())
+	{
+		return OutData;
+	}
+	
+	FMemoryReaderView Reader(InBytes);
+	FCborStructDeserializerBackend Deserializer(Reader);
+	UScriptStruct* InType = TBaseStructure<FLiveLinkOpenTrackIOData>::Get();
+
+	FLiveLinkOpenTrackIOData Data;
+	const bool bReadCbor = FStructDeserializer::Deserialize(&Data, *InType, Deserializer) && !Reader.GetError();
+
+	if (bReadCbor)
+	{
+		OutData.Emplace(MoveTemp(Data));
+	}
+
+	return OutData;
+}
+
+/**
+ * OpenTrackIO specifies using Fletcher-16 algorithm for checksum. This follows from the reference
+ * code for Fletcher-16. The article on wikipedia suggests some possible performance improvements
+ * but we opted for the unoptimized version: https://en.wikipedia.org/wiki/Fletcher%27s_checksum
+ */
+uint16_t CalculateChecksum(TArrayView<uint8> HeaderBytes, TArrayView<uint8> PayloadBytes)
+{
+	uint8 Sum1 = 0;
+	uint8 Sum2 = 0;
+	for (uint8 Byte : HeaderBytes)
+	{
+		Sum1 = (Sum1 + Byte) % 256;
+		Sum2 = (Sum2 + Sum1) % 256;
+	}
+	for (uint8 Byte : PayloadBytes)
+	{
+		Sum1 = (Sum1 + Byte) % 256;
+		Sum2 = (Sum2 + Sum1) % 256;		
+	}
+	return ((uint16) Sum2 << 8) | (uint16) Sum1;
+}
+
+/**
+ * These are the defined values in the OpenTrackIO specification.
+ */
+const uint32 OpenTrackIOHeaderId = 0x4F54726B;
+const uint8 JSON_Encoding = 0x01;
+const uint8 CBOR_Encoding = 0x02;
+
+uint16 ByteSwap(uint16 InVal)
+{
+#if PLATFORM_LITTLE_ENDIAN
+	return BYTESWAP_ORDER16(InVal);
+#else
+	return InVal;
+#endif
+}
+
+uint32 ByteSwap(uint32 InVal)
+{
+#if PLATFORM_LITTLE_ENDIAN
+	return BYTESWAP_ORDER32(InVal);
+#else
+	return InVal;
+#endif
+}
+	
+inline FArchive& operator<<(FArchive& Ar, FLiveLinkOpenTrackIODatagramHeader& Header)
+{
+	
+	uint32 Identifier;
+	Ar << Identifier;
+	Header.Identifier = ByteSwap(Identifier);
+
+	// 1 byte values.  No byte swapping necessary.
+	Ar << Header.Reserved;
+	Ar << Header.Encoding;
+
+	uint16 SequenceNumber;
+	Ar << SequenceNumber;
+	Header.SequenceNumber = ByteSwap(SequenceNumber);
+
+	uint32 SegmentOffset;
+	Ar << SegmentOffset;
+	Header.SegmentOffset = ByteSwap(SegmentOffset);
+
+	uint16 SegmentAndPayloadLength;
+	Ar << SegmentAndPayloadLength;
+	Header.LastSegmentFlagAndPayloadLength = ByteSwap(SegmentAndPayloadLength);
+
+	uint16 Checksum;
+	Ar << Checksum;
+	Header.Checksum = ByteSwap(Checksum);
+	
+	return Ar;
+}
+
+TOptional<FOpenTrackIOHeaderWithData> ParseOpenTrackIOPacket(TArrayView<uint8> Bytes)
+{
+	const uint32 HeaderSize = sizeof(FLiveLinkOpenTrackIODatagramHeader);
+
+	TOptional<FOpenTrackIOHeaderWithData> OutData;
+	/** No way this ia a valid datagram if we can't fit it in our header. */
+	if (Bytes.Num() < HeaderSize)
+	{
+		return OutData;
+	}
+
+	/** Check the header for validity first before trying to read the rest of the packet. */
+	{
+		FMemoryReaderView Reader(Bytes);
+		
+		uint32 Identifier;
+		Reader << Identifier;
+		Identifier = ByteSwap(Identifier);
+		if (Identifier != OpenTrackIOHeaderId)
+		{
+			UE_LOG(LogLiveLinkOpenTrackIO, Error, TEXT("Invalid OpenTrackIO header."));
+			return OutData;
+		}
+	}
+
+	/** Now read in the header to prepare parsing the payload. */
+	FMemoryReaderView Reader(Bytes);
+	FLiveLinkOpenTrackIODatagramHeader Header;
+	Reader << Header;
+
+	/** Confirm that the payload size is within the number of bytes given to us. */
+	uint16 PayloadSize = Header.GetPayloadSize();
+	if (PayloadSize > Bytes.Num())
+	{
+		/** The header does not match the inbound data. Abort! */
+		UE_LOG(LogLiveLinkOpenTrackIO, Error, TEXT("Payload size does not match provided OpenTrackIO packet header."));
+		return OutData;
+	}
+
+	TArrayView<uint8> HeaderMinusChecksum = Bytes.Slice(0, HeaderSize-2);
+	TArrayView<uint8> Payload = Bytes.Slice(HeaderSize, PayloadSize);
+
+	/** Use a fletcher-16 algorithm to calculate the checksum of the data. */
+	uint16_t Checksum = CalculateChecksum(HeaderMinusChecksum, Payload);
+	if (Checksum != Header.Checksum)
+	{
+		UE_LOG(LogLiveLinkOpenTrackIO, Warning, TEXT("Failed to verify packet cheksum."));
+		return OutData;
+	}
+
+	/** Now attempt to parse the payload using either JSON or CBOR. */
+	TOptional<FLiveLinkOpenTrackIOData> Data;
+	if (Header.Encoding == JSON_Encoding)
+	{
+		FString JsonBlob;
+		FFileHelper::BufferToString(JsonBlob, Payload.GetData(), Payload.Num());
+		Data = ParseJsonBlob(JsonBlob);
+	}
+	else if (Header.Encoding == CBOR_Encoding)
+	{
+		Data = ParseCborBlob(Payload);
+	}
+	else
+	{
+		UE_LOG(LogLiveLinkOpenTrackIO, Error, TEXT("Unsupported encoding."));
+		return OutData;
+	}
+
+	if (Data)
+	{
+		if (Data->Protocol.IsSupported())
+		{
+			OutData.Emplace(FOpenTrackIOHeaderWithData{Header, *Data});
+		}
+		else
+		{
+			UE_LOG(LogLiveLinkOpenTrackIO, Warning, TEXT("Unsupported OpenTrackIO protocol."));
+		}
+	}
+
+	return OutData;
+}
+
+}
+
+
diff --git a/Engine/Plugins/VirtualProduction/LiveLinkOpenTrackIO/Source/LiveLinkOpenTrackIO/Private/LiveLinkOpenTrackIOParser.h b/Engine/Plugins/VirtualProduction/LiveLinkOpenTrackIO/Source/LiveLinkOpenTrackIO/Private/LiveLinkOpenTrackIOParser.h
new file mode 100644
--- /dev/null
+++ b/Engine/Plugins/VirtualProduction/LiveLinkOpenTrackIO/Source/LiveLinkOpenTrackIO/Private/LiveLinkOpenTrackIOParser.h
@@ -0,0 +1,32 @@
+// Copyright Epic Games, Inc. All Rights Reserved.
+
+#pragma once
+
+#include "LiveLinkOpenTrackIOTypes.h"
+#include "LiveLinkOpenTrackIODatagram.h"
+
+#include "Containers/ArrayView.h"
+
+using FOpenTrackIOHeaderWithData = TTuple<FLiveLinkOpenTrackIODatagramHeader, FLiveLinkOpenTrackIOData>;
+
+namespace UE::OpenTrackIO::Private
+{
+	/**
+	 * Try to parse a JSON string and return an optional FLiveLinkOpenTrackIOData type.  If parsing fails
+	 * a unset optional will be returned.  If parsing succeeds then data will be returned.
+	 *
+	 * It is assumed that the JSON blob conforms to the minimal JSON 
+	 */
+	TOptional<FLiveLinkOpenTrackIOData> ParseJsonBlob(const FString& JsonBlob);
+
+	/**
+	 * Try to parse a CBOR blob. Return an optional FLiveLinkOpenTrackIOData type if CBOR parsing succeeds.
+	 */
+	TOptional<FLiveLinkOpenTrackIOData> ParseCborBlob(TArrayView<uint8> ByteArray);
+
+	/**
+	 * Try to parse and read an OpenTrackIO packet.  If a packet read was successful then
+	 * a tuple will be returned that consists of the header for the datagram and the parsed payload.
+	 */
+	TOptional<FOpenTrackIOHeaderWithData> ParseOpenTrackIOPacket(TArrayView<uint8> Bytes);
+}
diff --git a/Engine/Plugins/VirtualProduction/LiveLinkOpenTrackIO/Source/LiveLinkOpenTrackIO/Private/LiveLinkOpenTrackIOSource.cpp b/Engine/Plugins/VirtualProduction/LiveLinkOpenTrackIO/Source/LiveLinkOpenTrackIO/Private/LiveLinkOpenTrackIOSource.cpp
new file mode 100644
--- /dev/null
+++ b/Engine/Plugins/VirtualProduction/LiveLinkOpenTrackIO/Source/LiveLinkOpenTrackIO/Private/LiveLinkOpenTrackIOSource.cpp
@@ -0,0 +1,508 @@
+// Copyright Epic Games, Inc. All Rights Reserved.
+
+#include "LiveLinkOpenTrackIOSource.h"
+
+#include "Engine/Engine.h"
+
+#include "LiveLinkTypes.h"
+#include "ILiveLinkClient.h"
+
+#include "LiveLinkOpenTrackIO.h"
+#include "LiveLinkOpenTrackIOSourceSettings.h"
+#include "LiveLinkOpenTrackIOParser.h"
+#include "LiveLinkOpenTrackIOTranscoder.h"
+
+#include "Roles/LiveLinkCameraRole.h"
+#include "LiveLinkLensRole.h"
+
+#include "HAL/Platform.h"
+#include "Misc/CoreDelegates.h"
+
+#include "Common/UdpSocketBuilder.h"
+#include "Common/UdpSocketReceiver.h"
+
+#include "Sockets.h"
+#include "SocketSubsystem.h"
+#include "Interfaces/IPv4/IPv4Address.h"
+#include "Interfaces/IPv4/IPv4Endpoint.h"
+
+#define LOCTEXT_NAMESPACE "LiveLinkOpenTrackIOSource"
+
+namespace UE::LiveLinkOpenTrackIO::Private
+{
+
+TOptional<FIPv4Endpoint> ParseEndpoint(const FString& InEndpointString)
+{
+	FIPv4Endpoint Endpoint;
+	bool bParsedAddr = FIPv4Endpoint::Parse(InEndpointString, Endpoint);
+	if (!bParsedAddr)
+	{
+		bParsedAddr = FIPv4Endpoint::FromHostAndPort(InEndpointString, Endpoint);
+	}
+	if (bParsedAddr)
+	{
+		// Detect 169.254.x.x addresses.
+		if (Endpoint.Address.A == 169 &&	Endpoint.Address.B == 254)
+		{
+			UE_LOG(LogLiveLinkOpenTrackIO, Warning, TEXT("Detected IPv4 address in the form of 169.254.x.x. This is a link assigned address and may prevent you from reach external endpoints. "));
+		}
+
+		return Endpoint;
+	}
+
+	return {};
+}
+
+void DoJoinMulticastGroup(const TSharedRef<FInternetAddr>& MulticastAddr, const TSharedPtr<FInternetAddr>& IpAddr, FSocket* MulticastSocket)
+{
+	if (!IpAddr.IsValid())
+	{
+		return;
+	}
+	const bool bJoinedGroup = MulticastSocket->JoinMulticastGroup(*MulticastAddr, *IpAddr);
+	if (bJoinedGroup)
+	{
+		UE_LOG(LogLiveLinkOpenTrackIO, Display, TEXT("Added local interface '%s' to multicast group '%s'"),
+			   *IpAddr->ToString(false), *MulticastAddr->ToString(true));
+	}
+	else
+	{
+		UE_LOG(LogLiveLinkOpenTrackIO, Warning, TEXT("Failed to join multicast group '%s' on detected local interface '%s'"),
+			   *MulticastAddr->ToString(true), *IpAddr->ToString(false));
+	}
+}
+
+void JoinedToGroup(const FIPv4Endpoint& UnicastEndpoint, const FIPv4Endpoint& MulticastEndpoint, FSocket* MulticastSocket)
+{
+#if PLATFORM_SUPPORTS_UDP_MULTICAST_GROUP
+	TSharedRef<FInternetAddr> MulticastAddr = MulticastEndpoint.ToInternetAddr();
+	if (UnicastEndpoint.Address == FIPv4Address::Any)
+	{
+		TArray<TSharedPtr<FInternetAddr>> LocapIps;
+		ISocketSubsystem::Get(PLATFORM_SOCKETSUBSYSTEM)->GetLocalAdapterAddresses(LocapIps);
+		for (const TSharedPtr<FInternetAddr>& LocalIp : LocapIps)
+		{
+			DoJoinMulticastGroup(MulticastAddr, LocalIp, MulticastSocket);
+		}
+
+		// GetLocalAdapterAddresses returns empty list when all network adapters are offline
+		if (LocapIps.Num() == 0)
+		{
+			bool bCanBindAll = false;
+			DoJoinMulticastGroup(MulticastAddr, ISocketSubsystem::Get(PLATFORM_SOCKETSUBSYSTEM)->GetLocalHostAddr(*GLog, bCanBindAll), MulticastSocket);
+		}
+	}
+	else
+	{
+		TSharedRef<FInternetAddr> UnicastAddr = UnicastEndpoint.ToInternetAddr();
+		DoJoinMulticastGroup(MulticastAddr, UnicastAddr , MulticastSocket);
+	}
+#endif
+}
+
+}
+
+FLiveLinkOpenTrackIOSource::FLiveLinkOpenTrackIOSource(FLiveLinkOpenTrackIOConnectionSettings InConnectionSettings)
+	: ConnectionSettings(MoveTemp(InConnectionSettings))
+{
+	SourceType = LOCTEXT("SourceType_OpenTrackIO", "OpenTrack I/O");
+	SourceMachineName = LOCTEXT("LocalMachine", "Local");
+	SourceStatus = LOCTEXT("Startup", "Waiting for startup...");
+	DeferredStartDelegateHandle = FCoreDelegates::OnEndFrame.AddRaw(this, &FLiveLinkOpenTrackIOSource::Start);
+}
+
+FLiveLinkOpenTrackIOSource::~FLiveLinkOpenTrackIOSource()
+{
+	CloseSockets();
+
+	if (DeferredStartDelegateHandle.IsValid())
+	{
+		FCoreDelegates::OnEndFrame.Remove(DeferredStartDelegateHandle);
+	}
+}
+
+void FLiveLinkOpenTrackIOSource::ReceiveClient(ILiveLinkClient* InClient, FGuid InSourceGuid)
+{
+	Client = InClient;
+	SourceGuid = InSourceGuid;
+}
+
+bool FLiveLinkOpenTrackIOSource::IsSourceStillValid() const
+{
+	const bool bIsSourceValid = !Stopping && GameThreadConnectionState == ELiveLinkOpenTrackIOState::Receiving; 
+	return true;
+}
+
+bool FLiveLinkOpenTrackIOSource::RequestSourceShutdown()
+{
+	Stop();
+
+	return true;
+}
+
+void FLiveLinkOpenTrackIOSource::Update()
+{
+	if (GameThreadConnectionState != ConnectionState)
+	{
+		GameThreadConnectionState = ConnectionState;
+
+		switch (GameThreadConnectionState)
+		{
+			case ELiveLinkOpenTrackIOState::ResetRequested:
+				SourceStatus = LOCTEXT("Resetting", "Resetting source.");
+				break;
+			case ELiveLinkOpenTrackIOState::EndpointsReady:
+				SourceStatus = LOCTEXT("EndpointsReady", "Starting socket setup.");
+				break;
+			case ELiveLinkOpenTrackIOState::Receiving:
+				break;
+			default:
+				SourceStatus = LOCTEXT("NotStarted", "Not started");
+		};
+	}
+
+	if (GameThreadConnectionState == ELiveLinkOpenTrackIOState::Receiving)
+	{
+		if (OpenTrackIOPacketCount > 0)
+		{
+			SourceStatus = LOCTEXT("Receiving", "Receiving.");
+		}
+		else
+		{
+			SourceStatus = LOCTEXT("WaitingForData", "Waiting for data.");
+		}
+	}
+}
+
+void FLiveLinkOpenTrackIOSource::Start()
+{
+	check(DeferredStartDelegateHandle.IsValid());
+
+	FCoreDelegates::OnEndFrame.Remove(DeferredStartDelegateHandle);
+	DeferredStartDelegateHandle.Reset();
+
+	SourceStatus = LOCTEXT("Initialization", "Initializing receivers...");
+	ThreadName = "LiveLinkOpenTrackIO Receiver";
+	ThreadName.AppendInt(FAsyncThreadIndex::GetNext());
+
+	Thread.Reset(FRunnableThread::Create(this, *ThreadName, 128 * 1024, TPri_Normal, FPlatformAffinity::GetPoolThreadMask()));
+}
+
+void FLiveLinkOpenTrackIOSource::Stop()
+{
+	Stopping = true;
+}
+
+void FLiveLinkOpenTrackIOSource::ParseEndpointFromSourceSettings()
+{
+	check(IsValid(SavedSourceSettings));
+	
+	using namespace UE::LiveLinkOpenTrackIO::Private; 
+	if (TOptional<FIPv4Endpoint> InUniEndpoint = ParseEndpoint(SavedSourceSettings->UnicastEndpoint))
+	{
+		if (TOptional<FIPv4Endpoint> InMultiEndpoint = ParseEndpoint(SavedSourceSettings->MulticastEndpoint))
+		{
+			UnicastEndpoint = *InUniEndpoint;
+			MulticastEndpoint = *InMultiEndpoint;
+			ConnectionState = ELiveLinkOpenTrackIOState::EndpointsReady;
+		}
+	}
+}
+
+void FLiveLinkOpenTrackIOSource::InitializeSettings(ULiveLinkSourceSettings* Settings)
+{
+	ILiveLinkSource::InitializeSettings(Settings);
+
+	if (ULiveLinkOpenTrackIOSourceSettings* OpenTrackIOSettings = Cast<ULiveLinkOpenTrackIOSourceSettings>(Settings))
+	{
+		SavedSourceSettings = OpenTrackIOSettings;
+
+		ParseEndpointFromSourceSettings();
+	}
+	else
+	{
+		UE_LOG(LogLiveLinkOpenTrackIO, Error, TEXT("Invalid source settings."));
+	}
+}
+
+void FLiveLinkOpenTrackIOSource::ResetSockets()
+{
+	SourceStatus = LOCTEXT("Resetting", "Resetting...");
+	ConnectionState = ELiveLinkOpenTrackIOState::ResetRequested;
+	if (IsValid(SavedSourceSettings))
+	{
+		SavedSourceSettings->AssignedUnicastEndpoint = TEXT("");
+	}
+
+	OpenTrackIOPacketCount = 0;
+}
+
+void FLiveLinkOpenTrackIOSource::OnSettingsChanged(ULiveLinkSourceSettings* Settings, const FPropertyChangedEvent& PropertyChangedEvent)
+{
+	ILiveLinkSource::OnSettingsChanged(Settings, PropertyChangedEvent);
+
+	const FProperty* const MemberProperty = PropertyChangedEvent.MemberProperty;
+	const FProperty* const Property = PropertyChangedEvent.Property;
+	if (Property && MemberProperty && (PropertyChangedEvent.ChangeType != EPropertyChangeType::Interactive))
+	{
+		if (ensureMsgf(IsValid(SavedSourceSettings), TEXT("The LiveLinkOpenTrackIO source settings have become invalid.  Please re-create the source.")))
+		{
+			// If the user has changed the endpoint addresses then we need to reset the socket connection and re-establish receivers. 
+			const FName PropertyName = Property->GetFName();
+			const FName MemberPropertyName = MemberProperty->GetFName();
+			if (PropertyName == GET_MEMBER_NAME_CHECKED(ULiveLinkOpenTrackIOSourceSettings, UnicastEndpoint)
+				|| PropertyName == GET_MEMBER_NAME_CHECKED(ULiveLinkOpenTrackIOSourceSettings, MulticastEndpoint))
+			{
+				/** Start the reset process. */
+				ResetSockets();
+			}
+		}
+	}
+}
+	
+double FLiveLinkOpenTrackIOSource::DevicePollWaitTime() const
+{
+	const double UpdateRateInHz = LocalUpdateRateInHz > 0 ? (double) LocalUpdateRateInHz : 60.0;
+	return 1.0 / UpdateRateInHz;
+}
+
+void FLiveLinkOpenTrackIOSource::HandleInboundData(const TSharedPtr<FArrayReader, ESPMode::ThreadSafe>& Data, const FIPv4Endpoint& InSender)
+{
+	if (Stopping)
+	{
+		return;
+	}
+
+	Inbound.Enqueue(FInboundMessage{ Data, InSender });
+
+	/** TODO Rather than a busy loop we should use FEvent to trigger thread wake-up. */
+}
+
+bool FLiveLinkOpenTrackIOSource::ConsumeInboundData()
+{
+	bool bProcessedData = false;
+	FInboundMessage Message;
+	while (Inbound.Dequeue(Message))
+	{
+		if (!Message.Data.IsValid())
+		{
+			continue;
+		}
+
+		TArrayView<uint8> MessageData(Message.Data->GetData(), Message.Data->Num());
+		TOptional<FOpenTrackIOHeaderWithData> ParsedData = UE::OpenTrackIO::Private::ParseOpenTrackIOPacket(MessageData);
+		if (ParsedData)
+		{
+			bProcessedData = true;
+			OpenTrackIOPacketCount++;
+
+			PushDataToLiveLink_AnyThread(ParsedData->Get<0>() /** Header */, ParsedData->Get<1>() /** Body */);
+		}
+		else
+		{
+			UE_LOG(LogLiveLinkOpenTrackIO, Display, TEXT("Failed to handle inbound data."));
+		}
+	}
+	return bProcessedData;
+}
+
+void FLiveLinkOpenTrackIOSource::PushDataToLiveLink_AnyThread(const FLiveLinkOpenTrackIODatagramHeader& Header, const FLiveLinkOpenTrackIOData& InData)
+{
+	FLiveLinkOpenTrackIOCache* Cache;
+	if (OpenTrackIOCacheMap.Contains(InData.SourceId))
+	{
+		Cache = OpenTrackIOCacheMap.FindChecked(InData.SourceId).Get();
+	}
+	else
+	{
+		Cache = OpenTrackIOCacheMap.Add(InData.SourceId, MakePimpl<FLiveLinkOpenTrackIOCache>()).Get();
+	}
+
+	/** Check the static data camera fields on the inbound struct and see if we have new data to store / update. */
+	if (TOptional<FLiveLinkStaticDataStruct> StaticData = Cache->TryAddCamera(InData))
+	{
+		FName SubjectName = Cache->GetSubjectNameFromData(InData);
+
+		Client->PushSubjectStaticData_AnyThread({ SourceGuid, SubjectName }, ULiveLinkLensRole::StaticClass(), MoveTemp(*StaticData));
+	}
+
+	TArray<FSubjectAndFrameDataTuple> FrameData = Cache->TryReadFrameData(InData);
+	/**
+	 * This loop will consume the data in FrameData and move it into the Live Link thread. You should consider the data invalid
+	 * after this loop completes.
+	 */
+	for (FSubjectAndFrameDataTuple& Item : FrameData)
+	{
+		Client->PushSubjectFrameData_AnyThread({ SourceGuid, Item.SubjectName }, MoveTemp(Item.FrameData));
+	}
+}
+
+void FLiveLinkOpenTrackIOSource::CloseSockets()
+{
+	MulticastReceiver.Reset();
+	UnicastReceiver.Reset();
+
+	ISocketSubsystem* SocketSubsystem = ISocketSubsystem::Get(PLATFORM_SOCKETSUBSYSTEM);
+	check(SocketSubsystem);
+
+	if (MulticastSocket)
+	{
+		SocketSubsystem->DestroySocket(MulticastSocket);
+		MulticastSocket = nullptr;
+	}
+	if (UnicastSocket)
+	{
+		SocketSubsystem->DestroySocket(UnicastSocket);
+		UnicastSocket = nullptr;
+	}
+
+	OpenTrackIOCacheMap.Reset();
+}
+
+bool FLiveLinkOpenTrackIOSource::OpenSockets()
+{
+	if (ConnectionState != ELiveLinkOpenTrackIOState::EndpointsReady)
+	{
+		return false;
+	}
+
+	// Create an IPv4 TCP Socket
+	const uint32 ReceiveBufferSize = 2 * 1024 * 1024;
+
+	UnicastSocket = FUdpSocketBuilder(TEXT("UdpOpenTrackIO_UnicastSocket"))
+		.AsNonBlocking()
+		.BoundToEndpoint(UnicastEndpoint)
+		.WithMulticastLoopback()
+		.WithReceiveBufferSize(ReceiveBufferSize);
+
+	if (UnicastSocket == nullptr)
+	{
+		UE_LOG(LogLiveLinkOpenTrackIO, Error, TEXT("Failed to create OpenTrackIO socket on %s"), *UnicastEndpoint.ToString());
+		return false;
+	}
+	else
+	{
+		int32 PortNo = UnicastSocket->GetPortNo();
+		FString IpPortStr = FString::Format(TEXT("{0}:{1}"), {UnicastEndpoint.Address.ToString(), PortNo});
+		UE_LOG(LogLiveLinkOpenTrackIO, Display, TEXT("OpenTrackIO unicast socket socket bound to '%s'."), *IpPortStr);
+
+		AsyncTask(ENamedThreads::GameThread, [IpPortStr, this]() {
+			SourceMachineName = FText::Format(FTextFormat(LOCTEXT("LocalMachineWithIP", "Local ({0})")), {FText::FromString(IpPortStr)});
+			if (IsValid(SavedSourceSettings))
+			{
+				SavedSourceSettings->AssignedUnicastEndpoint = IpPortStr;
+			}
+		});
+	}
+
+	/** Create a multicast socket.  This */
+	MulticastSocket = FUdpSocketBuilder(TEXT("UdpOpenTrackIO_MulticastSocket"))
+		.AsNonBlocking()
+		.AsReusable()
+#if PLATFORM_WINDOWS
+
+		/**
+		 * If multiple bus instances bind the same unicast ip:port combination (allowed as the socket is marked as reusable), then for each packet
+		 * sent to that ip:port combination, only one of the instances (at the discretion of the OS) will receive it. The instance that receives the
+		 * packet may vary over time, seemingly based on the congestion of its socket. This isn't the intended usage.
+		 * 
+		 * To allow traffic to be sent directly to unicast for discovery, set the interface and port for the unicast endpoint
+		 * However for legacy reason, keep binding this as well, although it might be unreliable in some cases
+		 */
+		.BoundToAddress(UnicastEndpoint.Address)
+#endif
+		.BoundToPort(MulticastEndpoint.Port)
+#if PLATFORM_SUPPORTS_UDP_MULTICAST_GROUP
+		.WithMulticastLoopback()
+		.WithMulticastInterface(UnicastEndpoint.Address)
+#endif
+		.WithReceiveBufferSize(ReceiveBufferSize);
+
+
+	if (MulticastSocket != nullptr)
+	{
+		UE::LiveLinkOpenTrackIO::Private::JoinedToGroup(UnicastEndpoint, MulticastEndpoint, MulticastSocket);
+	}
+	else
+	{
+		UE_LOG(LogLiveLinkOpenTrackIO, Warning, TEXT("StartTransport failed to create multicast socket on %s, joined to %s"), *UnicastEndpoint.ToString(), *MulticastEndpoint.ToString());
+		return false;
+	}
+
+	/** Create receivers for sockets above */
+	const FTimespan ThreadWaitTime = FTimespan::FromMilliseconds(100);
+	if (MulticastSocket != nullptr)
+	{
+		MulticastReceiver = MakeUnique<FUdpSocketReceiver>(MulticastSocket, ThreadWaitTime, TEXT("LiveLinkOpenTransportIO_MulticastReceiver"));
+		MulticastReceiver->OnDataReceived().BindRaw(this, &FLiveLinkOpenTrackIOSource::HandleInboundData);
+		MulticastReceiver->SetMaxReadBufferSize(2048);
+		MulticastReceiver->Start();
+	}
+
+	if (UnicastSocket)
+	{
+		UnicastReceiver = MakeUnique<FUdpSocketReceiver>(UnicastSocket, ThreadWaitTime, TEXT("LiveLinkOpenTransportIO_UnicastReceiver"));
+		UnicastReceiver->OnDataReceived().BindRaw(this, &FLiveLinkOpenTrackIOSource::HandleInboundData);
+		UnicastReceiver->SetMaxReadBufferSize(2048);
+		UnicastReceiver->Start();
+	}
+
+	ConnectionState = ELiveLinkOpenTrackIOState::Receiving;
+	
+	return true;
+}
+
+uint32 FLiveLinkOpenTrackIOSource::Run()
+{
+	double SleepDeltaTime = DevicePollWaitTime();
+
+	double LastDataReadTime = FPlatformTime::Seconds();
+	while (!Stopping)
+	{
+		// Make sure we do not exceed our local update rate.
+		ON_SCOPE_EXIT{
+			FPlatformProcess::Sleep(SleepDeltaTime);
+		};
+
+		if (!GIsRunning)
+		{
+			continue;
+		}
+
+		if (!bHasStartedSockets)
+		{
+			bHasStartedSockets = OpenSockets();
+			continue;
+		}
+
+		/** Read the incoming data queue and push that into Live Link. */
+		if (ConsumeInboundData())
+		{
+			LastDataReadTime = FPlatformTime::Seconds();
+		}
+
+		/** If the user has requested a connection reset. Close the socket and attempt to re-open on the next loop iteration. */
+		if (ConnectionState == ELiveLinkOpenTrackIOState::ResetRequested)
+		{
+			CloseSockets();
+
+			/** Re-parse the endpoints to reset the connection. This will set the flag to allow the socket to reconnect. */
+			ParseEndpointFromSourceSettings();
+			bHasStartedSockets = false;
+		}
+
+		if ((FPlatformTime::Seconds() - LastDataReadTime) > 1.0)
+		{
+			/** We have lost in-bound data.  Reset the socket read count. */
+			OpenTrackIOPacketCount = 0;
+		}
+
+		SleepDeltaTime = DevicePollWaitTime();
+	}
+
+	CloseSockets();
+	
+	return 0;
+}
+
+#undef LOCTEXT_NAMESPACE
diff --git a/Engine/Plugins/VirtualProduction/LiveLinkOpenTrackIO/Source/LiveLinkOpenTrackIO/Private/LiveLinkOpenTrackIOSource.h b/Engine/Plugins/VirtualProduction/LiveLinkOpenTrackIO/Source/LiveLinkOpenTrackIO/Private/LiveLinkOpenTrackIOSource.h
new file mode 100644
--- /dev/null
+++ b/Engine/Plugins/VirtualProduction/LiveLinkOpenTrackIO/Source/LiveLinkOpenTrackIO/Private/LiveLinkOpenTrackIOSource.h
@@ -0,0 +1,196 @@
+// Copyright Epic Games, Inc. All Rights Reserved.
+
+#pragma once
+
+#include "ILiveLinkSource.h"
+
+#include "Interfaces/IPv4/IPv4Endpoint.h"
+
+#include "LiveLinkOpenTrackIOConnectionSettings.h"
+#include "LiveLinkOpenTrackIOSourceSettings.h"
+
+#include "Misc/CoreMiscDefines.h"
+
+#include "Containers/MpscQueue.h"
+#include "Templates/UniquePtr.h"
+#include "Templates/PimplPtr.h"
+
+#include "Roles/LiveLinkTransformTypes.h"
+
+#include "Common/UdpSocketReceiver.h"
+
+#include "Delegates/IDelegateInstance.h"
+#include "MessageEndpoint.h"
+#include "IMessageContext.h"
+#include "HAL/ThreadSafeBool.h"
+#include "HAL/Runnable.h"
+
+struct ULiveLinkOpenTrackIOSettings;
+
+class IOpenTrackIO;
+class ILiveLinkClient;
+class FArrayReader;
+enum class ESPMode : uint8;
+template<class ObjectType, ESPMode InMode> class TSharedPtr;
+
+struct FLiveLinkOpenTrackIOData;
+struct FLiveLinkOpenTrackIODatagramHeader;
+struct FLiveLinkOpenTrackIOCache;
+
+enum class ELiveLinkOpenTrackIOState : uint8
+{
+	/** Default state on source initialization. */
+	NotStarted = 0,
+
+	/** We have parsed the endpoints data from the source settings and will prepare to open receiving connections */
+	EndpointsReady,
+
+	/** Sockets have been opened and we are listening for data. */
+	Receiving,
+
+	/** The user has changed endpoints or user requested a socket reset. */
+	ResetRequested
+};
+
+class LIVELINKOPENTRACKIO_API FLiveLinkOpenTrackIOSource : public ILiveLinkSource, public FRunnable, public TSharedFromThis<FLiveLinkOpenTrackIOSource>
+{
+public:
+
+	FLiveLinkOpenTrackIOSource(FLiveLinkOpenTrackIOConnectionSettings ConnectionSettings);
+
+	virtual ~FLiveLinkOpenTrackIOSource();
+
+	// Begin ILiveLinkSource Interface
+	virtual void ReceiveClient(ILiveLinkClient* InClient, FGuid InSourceGuid) override;
+	virtual void InitializeSettings(ULiveLinkSourceSettings* Settings) override;
+	virtual void OnSettingsChanged(ULiveLinkSourceSettings* Settings, const FPropertyChangedEvent& PropertyChangedEvent) override;
+
+	virtual bool IsSourceStillValid() const override;
+
+	virtual bool RequestSourceShutdown() override;
+
+	virtual void Update() override;
+	
+	virtual FText GetSourceType() const override { return SourceType; };
+	virtual FText GetSourceMachineName() const override { return SourceMachineName; }
+	virtual FText GetSourceStatus() const override { return SourceStatus; }
+
+	virtual TSubclassOf<ULiveLinkSourceSettings> GetSettingsClass() const override { return ULiveLinkOpenTrackIOSourceSettings::StaticClass(); }
+	// End ILiveLinkSource Interface
+
+	// Begin FRunnable Interface
+	virtual bool Init() override { return true; }
+	virtual uint32 Run() override;
+	void Start();
+	virtual void Stop() override;
+	virtual void Exit() override { }
+	// End FRunnable Interface
+
+private:
+
+	/** Create listening sockets for OpenTrack channels. Returns true if socket handles were successfully opened and receivers started. */
+	bool OpenSockets();
+
+	/** Close any open sockets. */
+	void CloseSockets();
+
+	/** Time in seconds to wait before polling the device values again. */
+	double DevicePollWaitTime() const;
+
+	/** Read the source settings and assign the endpoint addresses. */
+	void ParseEndpointFromSourceSettings();
+
+	/** Trigger the reset mechanism. New sockets will be recreated. */
+	void ResetSockets();
+
+	/** Delegate for handling inbound segments. */
+	void HandleInboundData(const TSharedPtr<FArrayReader, ESPMode::ThreadSafe>& Data, const FIPv4Endpoint& InSender);
+
+	/** Consume inbound data. Fetch any data packets off of the inbound socket queue and process them. Returns true if data was processed. */
+	bool ConsumeInboundData();
+
+	/** Copy the open track data into Live Link equivalent data. This will establish any static data if the InData specifies it. */
+	void PushDataToLiveLink_AnyThread(const FLiveLinkOpenTrackIODatagramHeader& Header, const FLiveLinkOpenTrackIOData& InData);
+	
+private:
+	struct FInboundMessage
+	{
+		/** Holds the segment data. */
+		TSharedPtr<FArrayReader, ESPMode::ThreadSafe> Data;
+
+		/** Holds the sender's network endpoint. */
+		FIPv4Endpoint Sender;
+	};
+
+	/** Pointer to implementation of our OpenTrackIO cache so that we can track camera and lens information. */
+	using FOpenTrackIOCachePtr = TPimplPtr<FLiveLinkOpenTrackIOCache>;
+
+	/** Cache of data relevant to open track protocol. It is keyed by the sourceId. We can support multiple source ids */
+	TMap<FString, FOpenTrackIOCachePtr> OpenTrackIOCacheMap;
+	
+	/** Holds the queue of inbound messages. */
+	TMpscQueue<FInboundMessage> Inbound;
+
+	/** The LiveLink client provided via ReceiveClient*/
+	ILiveLinkClient* Client = nullptr;
+
+	/** Our identifier in LiveLink */
+	FGuid SourceGuid;
+
+	/** Text objects that are reflected back in the Live Link UI. */
+	FText SourceType;
+	FText SourceMachineName;
+	FText SourceStatus;
+
+	/** Handle the source settings for this source. The IP addresses are stored so that they can be changed after source creation. */
+	TObjectPtr<ULiveLinkOpenTrackIOSourceSettings> SavedSourceSettings;
+
+	/** Holds the connection settings to use for this source. This is established when the source is first created. */
+	FLiveLinkOpenTrackIOConnectionSettings ConnectionSettings;
+
+	/** Holds the multicast socket. */
+	FSocket* MulticastSocket = nullptr;
+
+	/** Holds the unicast socket receiver. */
+	TUniquePtr<FUdpSocketReceiver> UnicastReceiver;
+	
+	/** Holds the multicast socket receiver. */
+	TUniquePtr<FUdpSocketReceiver> MulticastReceiver;
+
+	/** Holds the unicast socket. */
+	FSocket* UnicastSocket = nullptr;
+
+	/** Holds the local endpoint to receive messages. */
+	FIPv4Endpoint UnicastEndpoint;
+
+	/** Multicast Endpoint endpoint to receive messages. */
+	FIPv4Endpoint MulticastEndpoint;
+	
+	/** Threadsafe Bool for terminating the main thread loop */
+	FThreadSafeBool Stopping;
+	
+	/** Thread to run updates on */
+	TUniquePtr<FRunnableThread> Thread;
+
+	/** Name of the running updates thread. Stored for debuggablity across multiple sources. */
+	FString ThreadName;
+
+	/** We don't start immediately after source contruction. We wait for the first engine tick. */
+	FDelegateHandle DeferredStartDelegateHandle;
+
+	/** Rate at which we poll the inbound queues. */
+	uint32 LocalUpdateRateInHz = 60;
+
+	/** Bool to indicate if we have started listening connections. Only used in running thread. */
+	bool bHasStartedSockets = false;
+
+	/** Variable to track the connection state from the game thread so that we can update the source status text. */
+	ELiveLinkOpenTrackIOState GameThreadConnectionState = ELiveLinkOpenTrackIOState::NotStarted;
+
+	/** Connection status */
+	std::atomic<ELiveLinkOpenTrackIOState> ConnectionState = ELiveLinkOpenTrackIOState::NotStarted;
+
+	/** Counter to track the number of inbound OpenTrackIO packets received. */
+	std::atomic<uint64> OpenTrackIOPacketCount = 0;
+	
+};
diff --git a/Engine/Plugins/VirtualProduction/LiveLinkOpenTrackIO/Source/LiveLinkOpenTrackIO/Private/LiveLinkOpenTrackIOSourceFactory.cpp b/Engine/Plugins/VirtualProduction/LiveLinkOpenTrackIO/Source/LiveLinkOpenTrackIO/Private/LiveLinkOpenTrackIOSourceFactory.cpp
new file mode 100644
--- /dev/null
+++ b/Engine/Plugins/VirtualProduction/LiveLinkOpenTrackIO/Source/LiveLinkOpenTrackIO/Private/LiveLinkOpenTrackIOSourceFactory.cpp
@@ -0,0 +1,48 @@
+// Copyright Epic Games, Inc. All Rights Reserved.
+
+#include "LiveLinkOpenTrackIOSourceFactory.h"
+
+#include "ILiveLinkClient.h"
+#include "LiveLinkOpenTrackIOSource.h"
+#include "LiveLinkOpenTrackIOSourceSettings.h"
+
+#include "SLiveLinkOpenTrackIOSourceFactory.h"
+
+#define LOCTEXT_NAMESPACE "LiveLinkOpenTrackIOSourceFactory"
+
+FText ULiveLinkOpenTrackIOSourceFactory::GetSourceDisplayName() const
+{
+	return LOCTEXT("SourceDisplayName", "Live Link OpenTrackIO Source");
+}
+
+FText ULiveLinkOpenTrackIOSourceFactory::GetSourceTooltip() const
+{
+	return LOCTEXT("SourceTooltip", "Creates a Live Link OpenTrackIO source (https://www.opentrackio.org/).");
+}
+
+TSharedPtr<SWidget> ULiveLinkOpenTrackIOSourceFactory::BuildCreationPanel(FOnLiveLinkSourceCreated InOnLiveLinkSourceCreated) const
+{
+	return SNew(SLiveLinkOpenTrackIOSourceFactory)
+		.OnConnectionSettingsAccepted(FOnLiveLinkOpenTrackIOConnectionSettingsAccepted::CreateUObject(this, &ULiveLinkOpenTrackIOSourceFactory::CreateSourceFromSettings, InOnLiveLinkSourceCreated));
+}
+
+TSharedPtr<ILiveLinkSource> ULiveLinkOpenTrackIOSourceFactory::CreateSource(const FString& ConnectionString) const
+{
+	FLiveLinkOpenTrackIOConnectionSettings ConnectionSettings;
+	if (!ConnectionString.IsEmpty())
+	{
+		FLiveLinkOpenTrackIOConnectionSettings::StaticStruct()->ImportText(*ConnectionString, &ConnectionSettings, nullptr, PPF_None, GLog, TEXT("ULiveLinkOpenTrackIOSourceFactory"));
+	}
+	return MakeShared<FLiveLinkOpenTrackIOSource>(ConnectionSettings);
+}
+
+void ULiveLinkOpenTrackIOSourceFactory::CreateSourceFromSettings(FLiveLinkOpenTrackIOConnectionSettings InConnectionSettings, FOnLiveLinkSourceCreated OnSourceCreated) const
+{
+	FString ConnectionString;
+	FLiveLinkOpenTrackIOConnectionSettings::StaticStruct()->ExportText(ConnectionString, &InConnectionSettings, nullptr, nullptr, PPF_None, nullptr);
+
+	TSharedPtr<FLiveLinkOpenTrackIOSource> SharedPtr = MakeShared<FLiveLinkOpenTrackIOSource>(InConnectionSettings);
+	OnSourceCreated.ExecuteIfBound(SharedPtr, MoveTemp(ConnectionString));
+}
+
+#undef LOCTEXT_NAMESPACE
diff --git a/Engine/Plugins/VirtualProduction/LiveLinkOpenTrackIO/Source/LiveLinkOpenTrackIO/Private/LiveLinkOpenTrackIOSourceFactory.h b/Engine/Plugins/VirtualProduction/LiveLinkOpenTrackIO/Source/LiveLinkOpenTrackIO/Private/LiveLinkOpenTrackIOSourceFactory.h
new file mode 100644
--- /dev/null
+++ b/Engine/Plugins/VirtualProduction/LiveLinkOpenTrackIO/Source/LiveLinkOpenTrackIO/Private/LiveLinkOpenTrackIOSourceFactory.h
@@ -0,0 +1,24 @@
+// Copyright Epic Games, Inc. All Rights Reserved.
+
+#pragma once
+
+#include "LiveLinkSourceFactory.h"
+#include "LiveLinkOpenTrackIOSource.h"
+#include "LiveLinkOpenTrackIOSourceFactory.generated.h"
+
+UCLASS()
+class LIVELINKOPENTRACKIO_API ULiveLinkOpenTrackIOSourceFactory : public ULiveLinkSourceFactory
+{
+public:
+	GENERATED_BODY()
+
+	virtual FText GetSourceDisplayName() const override;
+	virtual FText GetSourceTooltip() const override;
+
+	virtual EMenuType GetMenuType() const override { return EMenuType::SubPanel; }
+	virtual TSharedPtr<SWidget> BuildCreationPanel(FOnLiveLinkSourceCreated OnLiveLinkSourceCreated) const override;
+	virtual TSharedPtr<ILiveLinkSource> CreateSource(const FString& ConnectionString) const override;
+
+private:
+	void CreateSourceFromSettings(FLiveLinkOpenTrackIOConnectionSettings ConnectionSettings, FOnLiveLinkSourceCreated OnSourceCreated) const;
+};
diff --git a/Engine/Plugins/VirtualProduction/LiveLinkOpenTrackIO/Source/LiveLinkOpenTrackIO/Private/LiveLinkOpenTrackIOTranscoder.cpp b/Engine/Plugins/VirtualProduction/LiveLinkOpenTrackIO/Source/LiveLinkOpenTrackIO/Private/LiveLinkOpenTrackIOTranscoder.cpp
new file mode 100644
--- /dev/null
+++ b/Engine/Plugins/VirtualProduction/LiveLinkOpenTrackIO/Source/LiveLinkOpenTrackIO/Private/LiveLinkOpenTrackIOTranscoder.cpp
@@ -0,0 +1,122 @@
+// Copyright Epic Games, Inc. All Rights Reserved.
+// 
+#include "LiveLinkOpenTrackIOTranscoder.h"
+
+#include "Misc/QualifiedFrameTime.h"
+#include "Roles/LiveLinkCameraTypes.h"
+#include "LiveLinkLensTypes.h"
+#include "Roles/LiveLinkCameraTypes.h"
+#include "Misc/App.h"
+
+template<typename T>
+bool AddStaticData(FName SourceId, const T& Item, TMap<FName, T>& Container)
+{
+	if (!Item.IsValid())
+	{
+		return false;
+	}
+	
+	T& FoundItem = Container.FindOrAdd(SourceId);
+	if (FoundItem != Item)
+	{
+		FoundItem = Item;
+		return true;
+	}
+
+	return false;
+}
+
+
+FName FLiveLinkOpenTrackIOCache::GetSubjectNameFromData(const FLiveLinkOpenTrackIOData& InData)
+{
+	return UE::OpenTrackIO::ConvertTypeToFName(InData.Static.Camera); 
+}
+
+TOptional<FLiveLinkStaticDataStruct> FLiveLinkOpenTrackIOCache::TryAddCamera(const FLiveLinkOpenTrackIOData& Data)
+{
+	FName SourceId = FName(Data.SourceId);
+	bool bIsNewStatic = AddStaticData(SourceId, Data.Static.Lens, StaticLenses)
+					 || AddStaticData(SourceId, Data.Static.Camera, StaticCameras);
+
+	if (!bIsNewStatic)
+	{
+		return {}; 
+	}
+
+	/** Lens class inherits from the camera data so we use that as our base for the static struct. */
+	FLiveLinkStaticDataStruct StaticDataStruct(FLiveLinkLensStaticData::StaticStruct());
+	FLiveLinkLensStaticData& LensData = *StaticDataStruct.Cast<FLiveLinkLensStaticData>();
+	LensData.bIsFocalLengthSupported = true;
+	LensData.bIsApertureSupported = true;
+	LensData.bIsFocusDistanceSupported = true;
+	LensData.FilmBackHeight = Data.Static.Camera.ActiveSensorPhysicalDimensions.Height;
+	LensData.FilmBackWidth = Data.Static.Camera.ActiveSensorPhysicalDimensions.Width;
+
+	/** OpenTrackIO only allows for spherical models right now. */
+	LensData.LensModel = FName(TEXT("Spherical"));
+	
+	return StaticDataStruct;
+}
+
+TArray<FSubjectAndFrameDataTuple> FLiveLinkOpenTrackIOCache::TryReadFrameData(const FLiveLinkOpenTrackIOData& Data)
+{
+	TArray<FSubjectAndFrameDataTuple> Result;
+	if (Data.Transforms.IsEmpty())
+	{
+		return Result;
+	}
+	
+	/** Compute the final transform by concatentating the transforms in the frame data. */
+	FTransform FinalTransform; 
+	for (const FLiveLinkOpenTrackIOTransform& Transform : Data.Transforms)
+	{
+		FinalTransform *= Transform.GetTransform();
+	}
+
+	FQualifiedFrameTime SampleTC;
+	if (Data.Timing.IsDefault())
+	{
+		const TOptional<FQualifiedFrameTime> CurrentFrameTime = FApp::GetCurrentFrameTime();
+		if (CurrentFrameTime)
+		{
+			SampleTC = *CurrentFrameTime;
+		}
+	}
+	else
+	{
+		SampleTC = Data.Timing.Timecode.GetQualifiedFrameTime();
+	}
+
+	TArray<FName> Keys;
+	StaticCameras.GenerateKeyArray(Keys);
+
+	double WorldTime = FPlatformTime::Seconds();
+	for (const FName SourceId : Keys)
+	{
+		// We require cameras to exist but it is optional to have a lens.
+		// 
+		FLiveLinkOpenTrackIOStaticCamera& Camera = StaticCameras.FindChecked(SourceId);
+		FLiveLinkOpenTrackIOStaticLens* Lens = StaticLenses.Find(SourceId);
+		
+		FLiveLinkFrameDataStruct FrameDataStruct = FLiveLinkFrameDataStruct(FLiveLinkLensFrameData::StaticStruct());
+		FLiveLinkLensFrameData& CameraAndLensData = *FrameDataStruct.Cast<FLiveLinkLensFrameData>();
+
+		CameraAndLensData.Transform = FinalTransform;
+		CameraAndLensData.MetaData.SceneTime = SampleTC;
+		CameraAndLensData.WorldTime = WorldTime;
+
+		CameraAndLensData.FocusDistance = Data.Lens.FocusDistance;
+		CameraAndLensData.Aperture = Data.Lens.FStop;
+		CameraAndLensData.FocalLength = Data.Lens.PinholeFocalLength;
+
+		if (Lens)
+		{
+			CameraAndLensData.PrincipalPoint = Data.Lens.DistortionOffset.AsVector(); 
+		}
+
+		FName SubjectName = UE::OpenTrackIO::ConvertTypeToFName(Camera);
+		Result.Add({SubjectName, MoveTemp(FrameDataStruct)});
+	}
+
+	return Result;
+}
diff --git a/Engine/Plugins/VirtualProduction/LiveLinkOpenTrackIO/Source/LiveLinkOpenTrackIO/Private/LiveLinkOpenTrackIOTranscoder.h b/Engine/Plugins/VirtualProduction/LiveLinkOpenTrackIO/Source/LiveLinkOpenTrackIO/Private/LiveLinkOpenTrackIOTranscoder.h
new file mode 100644
--- /dev/null
+++ b/Engine/Plugins/VirtualProduction/LiveLinkOpenTrackIO/Source/LiveLinkOpenTrackIO/Private/LiveLinkOpenTrackIOTranscoder.h
@@ -0,0 +1,37 @@
+// Copyright Epic Games, Inc. All Rights Reserved.
+
+#pragma once
+
+#include "LiveLinkOpenTrackIOTypes.h"
+
+#include "LiveLinkTypes.h"
+
+#include "Containers/Map.h"
+#include "Misc/Optional.h"
+
+struct FSubjectAndFrameDataTuple
+{
+	FName SubjectName;
+	FLiveLinkFrameDataStruct FrameData;
+};
+
+struct FLiveLinkOpenTrackIOCache
+{
+	TMap<FName, FLiveLinkOpenTrackIOStaticCamera> StaticCameras;
+	TMap<FName, FLiveLinkOpenTrackIOStaticLens>   StaticLenses;
+
+	FName GetSubjectNameFromData(const FLiveLinkOpenTrackIOData& Data);
+	
+	/**
+	 * Attempt to add this camera to our cache. If the camera does not exist or the is the same make / model but different values then the static
+	 * struct is returned.
+	 */
+	TOptional<FLiveLinkStaticDataStruct> TryAddCamera(const FLiveLinkOpenTrackIOData& StaticData);
+
+	/**
+	 * Attempt to read the per-frame data for the given OpenTrackIO sample. If valid data is present then a constructed frame data will be provided
+	 * for each subject in the cache.
+	 */
+	TArray<FSubjectAndFrameDataTuple> TryReadFrameData(const FLiveLinkOpenTrackIOData& Data);
+	
+};
diff --git a/Engine/Plugins/VirtualProduction/LiveLinkOpenTrackIO/Source/LiveLinkOpenTrackIO/Private/LiveLinkOpenTrackIOTypes.h b/Engine/Plugins/VirtualProduction/LiveLinkOpenTrackIO/Source/LiveLinkOpenTrackIO/Private/LiveLinkOpenTrackIOTypes.h
new file mode 100644
--- /dev/null
+++ b/Engine/Plugins/VirtualProduction/LiveLinkOpenTrackIO/Source/LiveLinkOpenTrackIO/Private/LiveLinkOpenTrackIOTypes.h
@@ -0,0 +1,1244 @@
+// Copyright Epic Games, Inc. All Rights Reserved.
+#pragma once
+
+#include "Misc/QualifiedFrameTime.h"
+#include "UObject/ObjectMacros.h"
+#include "Containers/Array.h"
+
+#include "Math/Matrix.h"
+#include "Math/Transform.h"
+
+#include "LiveLinkOpenTrackIOTypes.generated.h"
+
+/**
+ * Represents frame rate type from OpenTrack I/O. Value is represented in hertz
+ */
+USTRUCT()
+struct FLiveLinkOpenTrackIOFrameRate
+{
+	GENERATED_BODY()
+
+	/** Returns true if the struct is equivalent to the default struct values. */
+	bool IsDefault() const
+	{
+		static FLiveLinkOpenTrackIOFrameRate DefaultValue;
+		return DefaultValue == *this;
+	}
+	
+	bool operator==(const FLiveLinkOpenTrackIOFrameRate&) const = default;
+	bool operator!=(const FLiveLinkOpenTrackIOFrameRate&) const = default;
+
+	FFrameRate GetFrameRate() const
+	{
+		check(Num > 0);
+		return FFrameRate(static_cast<uint32>(Num), Denom);
+	}
+
+	UPROPERTY()
+	int32 Num = 24;
+
+	UPROPERTY()
+	uint32 Denom = 1;
+};
+
+/** Concepts to test for the existance of member variables on our structs. */
+template<typename T>
+concept CHasMakeMember = requires(T t)
+{ t.Make; };
+
+template<typename T>
+concept CHasModelMember = requires(T t)
+{ t.Model; };
+
+template<typename T>
+concept CHasLabelMember = requires(T t)
+{ t.Label; };
+
+template <typename T>
+concept CHasModelAndMake = CHasMakeMember<T> && CHasModelMember<T>;
+
+namespace UE::OpenTrackIO
+{
+
+template<CHasModelAndMake T>
+FName ConvertTypeToFName(const T& Type)
+{
+	if constexpr(CHasLabelMember<T>)
+	{
+		return FName(*Type.Label);
+	}
+	else
+	{
+		FNameBuilder Builder;
+		Builder.Append(*Type.Make);
+		Builder.Append(TEXT("_"));
+		Builder.Append(*Type.Model);
+		return *Builder;
+	}
+}
+
+}
+/**
+ * Height and width of the active area of the camera sensor in millimeters
+ */
+USTRUCT()
+struct FLiveLinkOpenTrackIOSensorDimensions
+{
+	GENERATED_BODY()
+
+	/** Returns true if the struct is equivalent to the default struct values. */
+	bool IsDefault() const
+	{
+		static FLiveLinkOpenTrackIOSensorDimensions DefaultValue;
+		return DefaultValue == *this;
+	}
+	
+	bool operator==(const FLiveLinkOpenTrackIOSensorDimensions&) const = default;
+	bool operator!=(const FLiveLinkOpenTrackIOSensorDimensions&) const = default;
+	
+	UPROPERTY()
+	float Height = 43.0f;
+
+	UPROPERTY()
+	float Width = 24.0f;
+};
+
+/**
+ *  Photosite resolution of the active area of the camera sensor in pixels
+ */
+USTRUCT()
+struct FLiveLinkOpenTrackIOSensorResolution
+{
+	GENERATED_BODY()
+
+	/** Returns true if the struct is equivalent to the default struct values. */
+	bool IsDefault() const
+	{
+		static FLiveLinkOpenTrackIOSensorResolution DefaultValue;
+		return DefaultValue == *this;
+	}
+	
+	bool operator==(const FLiveLinkOpenTrackIOSensorResolution&) const = default;
+	bool operator!=(const FLiveLinkOpenTrackIOSensorResolution&) const = default;
+	
+	UPROPERTY()
+	int32 Height = 1920;
+
+	UPROPERTY()
+	int32 Width = 1080;
+};
+
+/**
+ * Nominal ratio of height to width of the image of an axis-aligned square captured by the camera sensor.
+ * It can be used to de-squeeze images but is not however an exact number over the entire captured
+ * area due to a lens’ intrinsic analog nature.
+ */
+USTRUCT()
+struct FLiveLinkOpenTrackIOAnamorphicSqueeze
+{
+	GENERATED_BODY()
+
+	/** Returns true if the struct is equivalent to the default struct values. */
+	bool IsDefault() const
+	{
+		static FLiveLinkOpenTrackIOAnamorphicSqueeze DefaultValue;
+		return DefaultValue == *this;
+	}
+	
+	bool operator==(const FLiveLinkOpenTrackIOAnamorphicSqueeze&) const = default;
+	bool operator!=(const FLiveLinkOpenTrackIOAnamorphicSqueeze&) const = default;
+	
+	UPROPERTY()
+	int32 Num = 1;
+
+	UPROPERTY()
+	uint32 Denom = 1;		
+};
+
+USTRUCT()
+struct FLiveLinkOpenTrackIOStaticCamera
+{
+	GENERATED_BODY()
+
+	/** Returns true if the struct is equivalent to the default struct values. */
+	bool IsDefault() const
+	{
+		static FLiveLinkOpenTrackIOStaticCamera DefaultValue;
+		return DefaultValue == *this;
+	}
+	
+	bool operator==(const FLiveLinkOpenTrackIOStaticCamera&) const = default;
+	bool operator!=(const FLiveLinkOpenTrackIOStaticCamera&) const = default;
+	
+	/** If the provided static camera specifies the correct properties to be considered valid. */
+	bool IsValid() const
+	{
+		return !Make.IsEmpty() && !Model.IsEmpty() && !SerialNumber.IsEmpty() && !Label.IsEmpty();
+	}
+	
+	/** Object representing the capture rate.  */
+	UPROPERTY()
+	FLiveLinkOpenTrackIOFrameRate CaptureFrameRate;
+
+	/** Object representing the sensor dimenions. */
+	UPROPERTY()
+	FLiveLinkOpenTrackIOSensorDimensions ActiveSensorPhysicalDimensions;
+
+	/** Object representing the sensor resolution.  */
+	UPROPERTY()
+	FLiveLinkOpenTrackIOSensorResolution ActiveSensorResolution;
+
+	/** Non-blank string naming camera manufacturer. */
+	UPROPERTY()
+	FString Make;
+
+	/** Non-blank string naming camera model. */
+	UPROPERTY()
+	FString Model;
+
+	/** Non-blank string for camera serial number. */
+	UPROPERTY()
+	FString SerialNumber;
+
+	/** Non-blank string identifying camera firmware version */
+	UPROPERTY()
+	FString FirmwareVersion;
+
+	/** Non-blank string containing user-determined camera identifier */
+	UPROPERTY()
+	FString Label;
+
+	/** Anamorphic squeeze ratio */
+	UPROPERTY()
+	FLiveLinkOpenTrackIOAnamorphicSqueeze AnamorphicSqueeze;
+
+	/** Arithmetic ISO scale as defined in ISO 12232 */
+	UPROPERTY()
+	uint32 IsoSpeed = 400;
+
+	/**
+	 * URN identifying the ASC Framing Decision List used by the camera.
+	 * Pattern:  ^urn:uuid:[0-9a-f]{8}-[0-9a-f]{4}-[0-9a-f]{4}-[0-9a-f]{4}-[0-9a-f]{12}$
+	 */
+	UPROPERTY()
+	FString FdlLink;
+
+	/**
+	 * Shutter speed as a fraction of the capture frame rate. The shutter speed (in units of 1/s) is equal
+	 * to the value of the parameter divided by 360 times the capture frame rate.
+	 */
+	UPROPERTY(meta = (ClampMin = "0.0", ClampMax = "360.0"))
+	float ShutterAngle = 0.0;
+};
+
+USTRUCT()
+struct FLiveLinkOpenTrackIOStaticLens
+{
+	GENERATED_BODY()
+
+	/** Returns true if the struct is equivalent to the default struct values. */
+	bool IsDefault() const
+	{
+		static FLiveLinkOpenTrackIOStaticLens DefaultValue;
+		return DefaultValue == *this;
+	}
+	
+	bool operator==(const FLiveLinkOpenTrackIOStaticLens&) const = default;
+	bool operator!=(const FLiveLinkOpenTrackIOStaticLens&) const = default;
+
+	/** If the provided static lens specifies the correct properties to be considered valid. */
+	bool IsValid() const
+	{
+		return !Make.IsEmpty() && !Model.IsEmpty() && !SerialNumber.IsEmpty();
+	}
+
+	/**
+	 * Static maximum overscan factor on lens distortion. This is an alternative to
+	 * providing dynamic overscan values each frame. Note it should be the maximum
+	 * of both projection-matrix-based and field-of-view-based rendering as per the OpenLensIO documentation.
+	 */
+	UPROPERTY()
+	float DistortionOverscanMax = 1.0;
+
+	/**
+	 * Static maximum overscan factor on lens undistortion. This is an alternative
+	 * to providing dynamic overscan values each frame. Note it should bethe maximum of both
+	 * projection-matrix-based and field-of-view-based rendering as per the OpenLensIO
+	 * documentation.
+	 */
+	UPROPERTY()
+	float UndistortionOverscanMax = 1.0;
+
+	/** Non-blank string naming lens manufacturer */
+	UPROPERTY()
+	FString Make;
+
+	/** Non-blank string identifying lens model */
+	UPROPERTY()
+	FString Model;
+
+	/** Non-blank string uniquely identifying the lens */
+	UPROPERTY()
+	FString SerialNumber;
+
+	/** Non-blank string identifying lens firmware version */
+	UPROPERTY()
+	FString FirmwareVersion;
+
+	/** Nominal focal length of the lens. The number printed on the side of a prime lens, e.g. 50 mm, and undefined in the case of a zoom lens. */
+	UPROPERTY()
+	float NominalFocalLength = 50;
+
+	/** List of free strings that describe the history of calibrations of the lens. */
+	UPROPERTY()
+	TArray<FString> CalibrationHistory;
+};
+
+USTRUCT()
+struct FLiveLinkOpenTrackIOStaticTracker
+{
+	GENERATED_BODY()
+
+	/** Returns true if the struct is equivalent to the default struct values. */
+	bool IsDefault() const
+	{
+		static FLiveLinkOpenTrackIOStaticTracker DefaultValue;
+		return DefaultValue == *this;
+	}
+	
+	bool operator==(const FLiveLinkOpenTrackIOStaticTracker&) const = default;
+	bool operator!=(const FLiveLinkOpenTrackIOStaticTracker&) const = default;
+
+	/** If the provided static lens specifies the correct properties to be considered valid. */
+	bool IsValid() const
+	{
+		return !Make.IsEmpty() && !Model.IsEmpty() && !SerialNumber.IsEmpty();
+	}
+
+	/** Non-blank string naming tracking device manufacturer */
+	UPROPERTY()
+	FString Make;
+
+	/**  Non-blank string identifying tracking device model */
+	UPROPERTY()
+	FString Model;
+
+	/** Non-blank string uniquely identifying the tracking device */
+	UPROPERTY()
+	FString SerialNumber;
+
+	/**  Non-blank string identifying tracking device firmware version */
+	UPROPERTY()
+	FString FirmwareVersion;
+};
+
+USTRUCT()
+struct FLiveLinkOpenTrackIOStaticDuration
+{
+	GENERATED_BODY()
+
+	/** Returns true if the struct is equivalent to the default struct values. */
+	bool IsDefault() const
+	{
+		static FLiveLinkOpenTrackIOStaticDuration DefaultValue;
+		return DefaultValue == *this;
+	}
+	
+	bool operator==(const FLiveLinkOpenTrackIOStaticDuration&) const = default;
+	bool operator!=(const FLiveLinkOpenTrackIOStaticDuration&) const = default;
+	
+	UPROPERTY()
+	int32 Num = 1;
+
+	UPROPERTY()
+	uint32 Denom = 1; 
+};
+
+/**
+ * Static data from a parsed OpenTrackIO datagram. Per the spec this will be sent periodically from the source.
+ */
+USTRUCT()
+struct FLiveLinkOpenTrackIOStaticData
+{
+	GENERATED_BODY()
+
+	/** Returns true if the struct is equivalent to the default struct values. */
+	bool IsDefault() const
+	{
+		static FLiveLinkOpenTrackIOStaticData DefaultValue;
+		return DefaultValue == *this;
+	}
+	
+	bool operator==(const FLiveLinkOpenTrackIOStaticData&) const = default;
+	bool operator!=(const FLiveLinkOpenTrackIOStaticData&) const = default;
+	
+	UPROPERTY()
+	FLiveLinkOpenTrackIOStaticDuration Duration;
+
+	UPROPERTY()
+	FLiveLinkOpenTrackIOStaticCamera Camera;
+
+	UPROPERTY()
+	FLiveLinkOpenTrackIOStaticLens Lens;
+
+	UPROPERTY()
+	FLiveLinkOpenTrackIOStaticTracker Tracker;
+};
+
+USTRUCT()
+struct FLiveLinkOpenTrackIOTracker
+{
+	GENERATED_BODY()
+
+	/** Returns true if the struct is equivalent to the default struct values. */
+	bool IsDefault() const
+	{
+		static FLiveLinkOpenTrackIOTracker DefaultValue;
+		return DefaultValue == *this;
+	}
+	
+	bool operator==(const FLiveLinkOpenTrackIOTracker&) const = default;
+	bool operator!=(const FLiveLinkOpenTrackIOTracker&) const = default;
+	
+	/** Non-blank string containing notes about tracking system */
+	UPROPERTY()
+	FString Notes;
+
+	/** Boolean indicating whether tracking system is recording data. */
+	UPROPERTY()
+	bool Recording = false;
+
+	/** Non-blank string describing the recording slate */
+	UPROPERTY()
+	FString Slate;
+
+	/** Non-blank string describing status of tracking system */
+	UPROPERTY()
+	FString Status;
+};
+
+/** 
+ * Representation of a OpenTrackIO Timestamp.
+ *
+ * The timestamp comprises a 48-bit unsigned integer (seconds), a 32-bit unsigned integer (nanoseconds).
+ *
+ * We use 64-bit unsigned values for storing seconds. 
+ */
+USTRUCT()
+struct FLiveLinkOpenTrackIOTimestamp
+{
+	GENERATED_BODY()
+
+	/** Returns true if the struct is equivalent to the default struct values. */
+	bool IsDefault() const
+	{
+		static FLiveLinkOpenTrackIOTimestamp DefaultValue;
+		return DefaultValue == *this;
+	}
+	
+	bool operator==(const FLiveLinkOpenTrackIOTimestamp&) const = default;
+	bool operator!=(const FLiveLinkOpenTrackIOTimestamp&) const = default;
+	
+	/** Although the schema specifies a 48 bit number. We allocate a full 64 bit integer here to store it. */
+	UPROPERTY()
+	uint64 Seconds = 0;
+
+	/** A 32-bit unsigned number for nanoseconds */
+	UPROPERTY()
+	uint32 Nanoseconds = 0;
+};
+
+USTRUCT()
+struct FLiveLinkOpenTrackIOSynchronizationOffsets
+{
+	GENERATED_BODY()
+
+	/** Returns true if the struct is equivalent to the default struct values. */
+	bool IsDefault() const
+	{
+		static FLiveLinkOpenTrackIOSynchronizationOffsets DefaultValue;
+		return DefaultValue == *this;
+	}
+	
+	bool operator==(const FLiveLinkOpenTrackIOSynchronizationOffsets&) const = default;
+	bool operator!=(const FLiveLinkOpenTrackIOSynchronizationOffsets&) const = default;
+	
+	UPROPERTY()
+	float Translation = 0.0f;
+
+	UPROPERTY()
+	float Rotation = 0.0f;
+
+	UPROPERTY()
+	float LensEncoders = 0.0f;
+};
+
+USTRUCT()
+struct FLiveLinkOpenTrackIO_PTPLeaderPriorities
+{
+	GENERATED_BODY()
+
+	/** Returns true if the struct is equivalent to the default struct values. */
+	bool IsDefault() const
+	{
+		static FLiveLinkOpenTrackIO_PTPLeaderPriorities DefaultValue;
+		return DefaultValue == *this;
+	}
+	
+	bool operator==(const FLiveLinkOpenTrackIO_PTPLeaderPriorities&) const = default;
+	bool operator!=(const FLiveLinkOpenTrackIO_PTPLeaderPriorities&) const = default;
+	
+	UPROPERTY()
+	uint8 Priority1 = 0;
+
+	UPROPERTY()
+	uint8 Priority2 = 0;
+};
+
+USTRUCT()
+struct FLiveLinkOpenTrackIO_PTP
+{
+	GENERATED_BODY()
+
+	/** Returns true if the struct is equivalent to the default struct values. */
+	bool IsDefault() const
+	{
+		static FLiveLinkOpenTrackIO_PTP DefaultValue;
+		return DefaultValue == *this;
+	}
+	
+	bool operator==(const FLiveLinkOpenTrackIO_PTP&) const = default;
+	bool operator!=(const FLiveLinkOpenTrackIO_PTP&) const = default;
+	
+	/** PTP Profile: Can be one the following values [IEEE Std 1588-2019, IEEE Std 802.1AS-2020, SMPTE ST2059-2:2021] */
+	UPROPERTY()
+	FString Profile = TEXT("SMPTE ST2059-2:2021");
+
+	UPROPERTY(meta = (ClampMin = "0"))
+	int8 Domain = 0;
+
+	/**
+	 * PTP Leader identity follows the pattern:
+	 * 
+	 * (?:^[0-9a-f]{2}(?::[0-9a-f]{2}){5}$)|(?:^[0-9a-f]{2}(?:-[0-9a-f]{2}){5}$)
+	 */
+	UPROPERTY()
+	FString LeaderIdentity; 
+
+	/**  Data structure for PTP synchronization priorities */
+	UPROPERTY()
+	FLiveLinkOpenTrackIO_PTPLeaderPriorities LeaderPriorities;
+
+	UPROPERTY()
+	float LeaderAccuracy = 0.0;
+
+	/** Enum string property [GNSS, Atomic clock, NTP] */
+	UPROPERTY()
+	FString LeaderTimeSource = TEXT("NTP");
+
+	UPROPERTY()
+	float MeanPathDelay = 0.0;
+
+	UPROPERTY()
+	uint32 Vlan = 0;
+};
+
+/**
+ * Object describing how the tracking device is synchronized for this sample.
+ *  
+ * frequency:
+ *
+ *  The frequency of a synchronization signal.This may differ from the sample frame rate for example in a genlocked tracking device. This is
+ *  not required if the synchronization source is PTP or NTP.
+ *
+ * locked: Is the tracking device locked to the synchronization source
+ *
+ * offsets: Offsets in seconds between sync and sample. Critical for e.g.
+ *   frame remapping, or when using different data sources for
+ *   position/rotation and lens encoding
+ *
+ * present: Is the synchronization source present (a synchronization source can be present but not locked if frame rates differ)
+ *
+ * ptp: If the synchronization source is a PTP leader, then this object
+ * 
+ * contains:
+ *
+ *  - "profile": Specifies the PTP profile in use. This defines the operational rules and parameters for synchronization. For example "SMPTE
+ *  ST2059-2:2021" for SMPTE 2110 based systems, or "IEEE Std 1588-2019" or "IEEE Std 802.1AS-2020" for industrial applications
+ *
+ *  - "domain": Identifies the PTP domain the device belongs to. Devices in the same domain can synchronize with each other
+ *  
+ *  - "leaderIdentity": The unique identifier (usually MAC address) of the current PTP leader (grandmaster)
+ *  
+ *  - "leaderPriorities": The priority values of the leader used in the Best Clock Algorithm (BMCA). Lower values indicate higher priority
+ *  
+ *  - "priority1": Static priority set by the administrator
+ *  
+ *  - "priority2": Dynamic priority based on the leader’s role or clock quality
+ *  
+ *  - "leaderAccuracy": The timing offset in seconds from the sample timestamp to the PTP timestamp
+ *  
+ *  - "meanPathDelay": The average round-trip delay between the device and the PTP leader, measured in seconds
+ *  
+ * source: The source of synchronization must be defined as one of the
+ * 
+ * following:
+ * 
+ *  - "vlan": Integer representing the VLAN ID for PTP traffic (e.g., 100 for
+ *  VLAN 100)
+ *  
+ *  - "leaderTimeSource": Indicates the leader’s source of time, such as GNSS, atomic clock, or NTP
+ *    
+ *  - "genlock": The tracking device has an external black/burst or tri-level analog sync signal that is triggering the capture of tracking samples
+ *  
+ *  - "videoIn": The tracking device has an external video signal that is triggering the capture of tracking samples
+ *  
+ *  - "ptp": The tracking device is locked to a PTP leader
+ *  
+ *  - "ntp": The tracking device is locked to an NTP server
+ *  
+ */
+USTRUCT()
+struct FLiveLinkOpenTrackIOSynchronization
+{
+	GENERATED_BODY()
+
+	/** Returns true if the struct is equivalent to the default struct values. */
+	bool IsDefault() const
+	{
+		static FLiveLinkOpenTrackIOSynchronization DefaultValue;
+		return DefaultValue == *this;
+	}
+	
+	bool operator==(const FLiveLinkOpenTrackIOSynchronization&) const = default;
+	bool operator!=(const FLiveLinkOpenTrackIOSynchronization&) const = default;
+	
+	UPROPERTY()
+	bool Locked = false;
+
+	/** Enum as string value can be "genlock", "videoIn", "ptp", ... */
+	UPROPERTY()
+	FString Source;
+
+	/** Synchonization rate. */
+	UPROPERTY()
+	FLiveLinkOpenTrackIOFrameRate Frequency;
+
+	UPROPERTY()
+	FLiveLinkOpenTrackIOSynchronizationOffsets Offsets;
+
+	UPROPERTY()
+	bool Present = false;
+
+	UPROPERTY()
+	FLiveLinkOpenTrackIO_PTP Ptp;
+};
+
+/**
+ * SMPTE timecode of the sample. Timecode is a standard for labeling
+ * individual frames of data in media systems and is useful for
+ * inter-frame synchronization. Frame rate is a rational number, allowing
+ * drop frame rates such as that colloquially called 29.97 to be
+ * represented exactly, as 30000/1001. The timecode frame rate may differ
+ * from the sample frequency. The zero-based sub-frame field allows for finer
+ * division of the frame, e.g. interlaced frames have two sub-frames,
+ * one per field.
+ */
+USTRUCT()
+struct FLiveLinkOpenTrackIOTimecode
+{
+	GENERATED_BODY()
+
+	/** Returns true if the struct is equivalent to the default struct values. */
+	bool IsDefault() const
+	{
+		static FLiveLinkOpenTrackIOTimecode DefaultValue;
+		return DefaultValue == *this;
+	}
+	
+	bool operator==(const FLiveLinkOpenTrackIOTimecode&) const = default;
+	bool operator!=(const FLiveLinkOpenTrackIOTimecode&) const = default;
+
+	FTimecode GetTimecode() const
+	{
+		double FPS = FrameRate.GetFrameRate().AsDecimal();
+		float SubFrameAsFloat = static_cast<float>(static_cast<double>(SubFrame) / FPS);
+		return FTimecode(Hours, Minutes, Seconds, Frames, SubFrameAsFloat, false);
+	}
+	
+	FQualifiedFrameTime GetQualifiedFrameTime() const
+	{
+		return FQualifiedFrameTime(GetTimecode(), FrameRate.GetFrameRate());
+	}
+	
+	UPROPERTY()	
+	int32 Hours = 0;
+
+	UPROPERTY()
+	int32 Minutes = 0;
+
+	UPROPERTY()
+	int32 Seconds = 0;
+
+	UPROPERTY()
+	int32 Frames = 0;
+
+	UPROPERTY()
+	FLiveLinkOpenTrackIOFrameRate FrameRate;
+	
+	UPROPERTY()
+	int32 SubFrame = 0;
+};
+
+USTRUCT()
+struct FLiveLinkOpenTrackIOTiming
+{
+	GENERATED_BODY()
+
+	/** Returns true if the struct is equivalent to the default struct values. */
+	bool IsDefault() const
+	{
+		static FLiveLinkOpenTrackIOTiming DefaultValue;
+		return DefaultValue == *this;
+	}
+	
+	bool operator==(const FLiveLinkOpenTrackIOTiming&) const = default;
+	bool operator!=(const FLiveLinkOpenTrackIOTiming&) const = default;
+	
+	/**
+	 *  Enumerated value indicating whether the sample transport mechanism provides
+	 *  inherent (’external’) timing, or whether the transport mechanism lacks inherent timing
+	 *  and so the sample must contain a PTP timestamp itself (’internal’) to carry timing information.
+	 **/ 
+	UPROPERTY()
+	FString Mode;
+
+	/** PTP timestamp of the data recording instant, provided for convenience during playback of e.g. pre-recorded tracking data. */
+	UPROPERTY()
+	FLiveLinkOpenTrackIOTimestamp RecordedTimestamp;
+
+	/**
+	 * Sample frame rate as a rational number. Drop frame rates such as 29.97 should be
+	 * represented as e.g. 30000/1001. In a variable rate system this should be estimated from the last
+	 * sample delta time.
+	 */
+	UPROPERTY()
+	FLiveLinkOpenTrackIOFrameRate SampleRate;
+
+	/** 
+	 * PTP timestamp of the data recording instant, provided for convenience during playback of e.g. pre-recorded tracking data.
+	 * The timestamp comprises a 48-bit unsigned integer (seconds), a 32-bit unsigned integer (nanoseconds)
+	 */
+	UPROPERTY()
+	FLiveLinkOpenTrackIOTimestamp SampleTimestamp;
+
+	/** Integer incrementing with each sample. */
+	UPROPERTY()
+	uint32 SequenceNumber = 0;
+
+	/** Synchronization object for the timing. */
+	UPROPERTY()
+	FLiveLinkOpenTrackIOSynchronization Synchronization;
+
+	/** SMTPE timecode of the sample. */
+	UPROPERTY()
+	FLiveLinkOpenTrackIOTimecode Timecode;
+};
+
+USTRUCT()
+struct FLiveLinkOpenTrackIOLens_DistortionOffset
+{
+	GENERATED_BODY()
+
+	/** Returns true if the struct is equivalent to the default struct values. */
+	bool IsDefault() const
+	{
+		static FLiveLinkOpenTrackIOLens_DistortionOffset DefaultValue;
+		return DefaultValue == *this;
+	}
+	
+	bool operator==(const FLiveLinkOpenTrackIOLens_DistortionOffset&) const = default;
+	bool operator!=(const FLiveLinkOpenTrackIOLens_DistortionOffset&) const = default;
+
+	FVector2D AsVector() const
+	{
+		return FVector2D(X,Y);
+	}
+	
+	UPROPERTY()
+	float X = 0;
+
+	UPROPERTY()
+	float Y = 0;
+};
+
+/**
+ * Encoders are represented in this way (as opposed to raw integer
+ * values) to ensure values remain independent of encoder resolution,
+ * minimum and maximum (at an acceptable loss of precision).
+ * These values are only relevant in lenses with end-stops that
+ * demarcate the 0 and 1 range.
+ * 
+ * Value should be provided in the following directions (if known):
+ * Focus:   0=infinite     1=closest
+ * Iris:    0=open         1=closed
+ * Zoom:    0=wide angle   1=telephoto
+ */
+USTRUCT()
+struct FLiveLinkOpenTrackIOLens_Encoders
+{
+	GENERATED_BODY()
+
+	/** Returns true if the struct is equivalent to the default struct values. */
+	bool IsDefault() const
+	{
+		static FLiveLinkOpenTrackIOLens_Encoders DefaultValue;
+		return DefaultValue == *this;
+	}
+	
+	bool operator==(const FLiveLinkOpenTrackIOLens_Encoders&) const = default;
+	bool operator!=(const FLiveLinkOpenTrackIOLens_Encoders&) const = default;
+	
+	UPROPERTY()
+	float Focus = 0;
+
+	UPROPERTY()
+	float Iris = 0;
+
+	UPROPERTY()
+	float Zoom = 0;
+};
+
+
+/**
+ * Raw encoder values for focus, iris and zoom.
+ * These values are dependent on encoder resolution and before any
+ * homing / ranging has taken place.
+ */
+USTRUCT()
+struct FLiveLinkOpenTrackIOLens_RawEncoders
+{
+	GENERATED_BODY()
+
+	/** Returns true if the struct is equivalent to the default struct values. */
+	bool IsDefault() const
+	{
+		static FLiveLinkOpenTrackIOLens_RawEncoders DefaultValue;
+		return DefaultValue == *this;
+	}
+	
+	bool operator==(const FLiveLinkOpenTrackIOLens_RawEncoders&) const = default;
+	bool operator!=(const FLiveLinkOpenTrackIOLens_RawEncoders&) const = default;
+	
+	UPROPERTY()
+	uint32 Focus = 0;
+
+	UPROPERTY()
+	uint32 Iris = 0;
+
+	UPROPERTY()
+	uint32 Zoom = 0;
+};
+
+/**
+ * Coefficients for calculating the exposure fall-off (vignetting) of a lens.
+ */
+USTRUCT()
+struct FLiveLinkOpenTrackIOLens_ExposureFalloff
+{
+	GENERATED_BODY()
+
+	/** Returns true if the struct is equivalent to the default struct values. */
+	bool IsDefault() const
+	{
+		static FLiveLinkOpenTrackIOLens_ExposureFalloff DefaultValue;
+		return DefaultValue == *this;
+	}
+	
+	bool operator==(const FLiveLinkOpenTrackIOLens_ExposureFalloff&) const = default;
+	bool operator!=(const FLiveLinkOpenTrackIOLens_ExposureFalloff&) const = default;
+	
+	UPROPERTY()
+	float A1 = 0;
+
+	UPROPERTY()
+	float A2 = 0;
+
+	UPROPERTY()
+	float A3 = 0;
+};
+
+USTRUCT()
+struct FLiveLinkOpenTrackIOLens_DistortionCoeff
+{
+	GENERATED_BODY()
+
+	/** Returns true if the struct is equivalent to the default struct values. */
+	bool IsDefault() const
+	{
+		static FLiveLinkOpenTrackIOLens_DistortionCoeff DefaultValue;
+		return DefaultValue == *this;
+	}
+	
+	bool operator==(const FLiveLinkOpenTrackIOLens_DistortionCoeff&) const = default;
+	bool operator!=(const FLiveLinkOpenTrackIOLens_DistortionCoeff&) const = default;
+	
+	UPROPERTY()
+	TArray<float> Radial;
+
+	UPROPERTY()
+	TArray<float> Tangential;
+
+	UPROPERTY()
+	float Overscan = 0.0;
+};
+
+USTRUCT()
+struct FLiveLinkOpenTrackIOLens
+{
+	GENERATED_BODY()
+
+	/** Returns true if the struct is equivalent to the default struct values. */
+	bool IsDefault() const
+	{
+		static FLiveLinkOpenTrackIOLens DefaultValue;
+		return DefaultValue == *this;
+	}
+	
+	bool operator==(const FLiveLinkOpenTrackIOLens&) const = default;
+	bool operator!=(const FLiveLinkOpenTrackIOLens&) const = default;
+	
+	/**
+	 * This list provides optional additional custom coefficients that can extend the existing lens model. The meaning
+	 * of and how these characteristics are to be applied to a virtual camera would require
+	 * negotiation between a particular producer and consumer.
+	 */
+	UPROPERTY()
+	TArray<float> Custom;
+
+	/**
+	 * A list of Distortion objects that each define the coefficients for calculating the distortion characteristics of a lens comprising radial
+	 * distortion coefficients of the spherical distortion (k1-N) and (optionally) the tangential distortion (p1-N). The key ’model’ names the
+	 * distortion model. Typical values for ’model’ include "Brown-Conrady D-U" when mapping distorted to undistorted coordinates, and "Brown-Conrady
+	 * U-D" when mapping undistorted to undistorted coordinates. If not provided, the default model is "Brown-Conrady D-U".
+	 */
+	UPROPERTY()
+	TArray<FLiveLinkOpenTrackIOLens_DistortionCoeff> Distortion;
+
+	/**
+	 *  Offset in x and y of the centre of distortion of the virtual camera in millimeter
+	 */
+	UPROPERTY()
+	FLiveLinkOpenTrackIOLens_DistortionOffset DistortionOffset;
+
+	/** Normalised real numbers (0-1) for focus, iris and zoom */
+	UPROPERTY()
+	FLiveLinkOpenTrackIOLens_Encoders Encoders;
+
+	/**
+	 * Offset of the entrance pupil relative to the nominal imaging plane
+	 * (positive if the entrance pupil is located on the side of the nominal
+	 * imaging plane that is towards the object, and negative otherwise).
+	 * Measured in meters as in a render engine it is often applied in the
+	 * virtual camera’s transform chain.
+	 */
+	UPROPERTY()
+	float  EntrancePupilOffset = 0;
+
+	/**
+	 * Coefficients for calculating the exposure fall-off (vignetting) of a lens
+	 */
+	UPROPERTY()
+	FLiveLinkOpenTrackIOLens_ExposureFalloff ExposureFalloff;
+
+	/** The linear f-number of the lens, equal to the focal length divided by the diameter of the entrance pupil. */
+	UPROPERTY()
+	float FStop = 2.8f;
+
+	/** Distance between the pinhole and the image plane in the simple CGI pinhole camera model. (millimeters) */
+	UPROPERTY()
+	float PinholeFocalLength = 35.0f;
+
+	/** Focus distance/position of the lens in meters. */ 
+	UPROPERTY()
+	float FocusDistance = 10.0f;
+
+	/**  Offset in x and y of the centre of perspective projection of the virtual camera */
+	UPROPERTY()
+	FLiveLinkOpenTrackIOLens_DistortionOffset ProjectionOffset;
+
+	/** Raw encoder values. */
+	UPROPERTY()
+	FLiveLinkOpenTrackIOLens_RawEncoders RawEncoders;
+
+	/** T-stop: Linear t-number of the lens, equal to the F-number of the lens divided by the square root of the transmittance of the lens. */
+	UPROPERTY()
+	float TStop = 0.0;
+};
+
+/**
+ * Name of the protocol in which the sample is being employed, and
+ * version of that protocol
+ */
+USTRUCT()
+struct FLiveLinkOpenTrackIOProtocol
+{
+	GENERATED_BODY()
+
+	/** Returns true if the struct is equivalent to the default struct values. */
+	bool IsDefault() const
+	{
+		static FLiveLinkOpenTrackIOProtocol DefaultValue;
+		return DefaultValue == *this;
+	}
+	
+	bool operator==(const FLiveLinkOpenTrackIOProtocol&) const = default;
+	bool operator!=(const FLiveLinkOpenTrackIOProtocol&) const = default;
+
+	bool IsSupported() const
+	{
+		/** Validate that we support this version of OpenTrackIO. We assume that patch release are compatible. */
+		return Name == TEXT("OpenTrackIO")
+			&& Version.Num() >= 2
+			&& Version[0] == 1
+			&& Version[1] == 0;
+	}
+
+	UPROPERTY()
+	FString Name;
+
+	// 3 numbers that represent the protocol version.  
+	UPROPERTY()
+	TArray<int8> Version;
+};
+
+/**
+ * Position of stage origin in global ENU and geodetic coordinates
+ * (E, N, U, lat0, lon0, h0). Note this may be dynamic if the stage is
+ * inside a moving vehicle. Units represented in meters. 
+ */
+USTRUCT()
+struct FLiveLinkOpenTrackIOGlobalStage
+{
+	GENERATED_BODY()
+
+	/** Returns true if the struct is equivalent to the default struct values. */
+	bool IsDefault() const
+	{
+		static FLiveLinkOpenTrackIOGlobalStage DefaultValue;
+		return DefaultValue == *this;
+	}
+	
+	bool operator==(const FLiveLinkOpenTrackIOGlobalStage&) const = default;
+	bool operator!=(const FLiveLinkOpenTrackIOGlobalStage&) const = default;
+
+	UPROPERTY()
+	float E = 0;
+
+	UPROPERTY()
+	float N = 0;
+
+	UPROPERTY()
+	float U = 0;
+
+	UPROPERTY()
+	float Lat0 = 0;
+
+	UPROPERTY()
+	float Lon0 = 0;
+
+	UPROPERTY()
+	float H0 = 0;
+};
+
+USTRUCT()
+struct FLiveLinkOpenTrackIO_Rotator
+{
+	GENERATED_BODY()
+	
+	/** Returns true if the struct is equivalent to the default struct values. */
+	bool IsDefault() const
+	{
+		static FLiveLinkOpenTrackIO_Rotator DefaultValue;
+		return DefaultValue == *this;
+	}
+	
+	bool operator==(const FLiveLinkOpenTrackIO_Rotator&) const = default;
+	bool operator!=(const FLiveLinkOpenTrackIO_Rotator&) const = default;
+
+	FRotator GetRotator() const
+	{
+		return FRotator(Tilt,Pan,Roll);
+	}
+	
+	UPROPERTY()
+	float Pan = 0;
+
+	UPROPERTY()
+	float Tilt = 0;
+
+	UPROPERTY()
+	float Roll = 0;
+};
+
+USTRUCT()
+struct FLiveLinkOpenTrackIO_XYZ
+{
+	GENERATED_BODY()
+
+	FLiveLinkOpenTrackIO_XYZ() = default;
+	FLiveLinkOpenTrackIO_XYZ(float InX, float InY, float InZ)
+		: X(InX), Y(InY), Z(InZ)
+	{
+	}
+
+	/** Returns true if the struct is equivalent to the default struct values. */
+	bool IsDefault() const
+	{
+		static FLiveLinkOpenTrackIO_XYZ DefaultValue;
+		return DefaultValue == *this;
+	}
+	
+	bool operator==(const FLiveLinkOpenTrackIO_XYZ&) const = default;
+	bool operator!=(const FLiveLinkOpenTrackIO_XYZ&) const = default;
+
+	FVector GetVector() const
+	{
+		return FVector(X,Y,Z);
+	}
+
+	UPROPERTY()
+	float X = 0;
+
+	UPROPERTY()
+	float Y = 0;
+
+	UPROPERTY()
+	float Z = 0;
+};
+
+USTRUCT()
+struct FLiveLinkOpenTrackIOTransform
+{
+	GENERATED_BODY()
+
+	/** Returns true if the struct is equivalent to the default struct values. */
+	bool IsDefault() const
+	{
+		static FLiveLinkOpenTrackIOTransform DefaultValue;
+		return DefaultValue == *this;
+	}
+	
+	bool operator==(const FLiveLinkOpenTrackIOTransform&) const = default;
+	bool operator!=(const FLiveLinkOpenTrackIOTransform&) const = default;
+
+	FTransform GetTransform() const
+	{
+		FTransform Transform(
+			Rotation.GetRotator(),
+			Translation.GetVector(),
+			Scale.GetVector());
+		return Transform;
+	}
+	
+	/** Translation part of a OpenTrackIO xform */
+	UPROPERTY()
+	FLiveLinkOpenTrackIO_XYZ Translation;
+	
+	/** Rotation part of a OpenTrackIO xform */
+	UPROPERTY()
+	FLiveLinkOpenTrackIO_Rotator Rotation;
+
+	/** Scale part of a OpenTrackIO xform */
+	UPROPERTY()
+	FLiveLinkOpenTrackIO_XYZ Scale = FLiveLinkOpenTrackIO_XYZ(1,1,1);
+
+	/** OpenTrackIO identifier (can be empty) */
+	UPROPERTY()
+	FString Id;
+};
+
+/**
+ * Data from a parsed header that was read by other CBOR or JSON format.
+ *
+ * Schema can be found here: https://www.opentrackio.org/schema.json
+ * 
+ */
+USTRUCT()
+struct FLiveLinkOpenTrackIOData
+{
+	GENERATED_BODY()
+
+	/** Returns true if the struct is equivalent to the default struct values. */
+	bool IsDefault() const
+	{
+		static FLiveLinkOpenTrackIOData DefaultValue;
+		return DefaultValue == *this;
+	}
+	
+	bool operator==(const FLiveLinkOpenTrackIOData&) const = default;
+	bool operator!=(const FLiveLinkOpenTrackIOData&) const = default;
+	
+	UPROPERTY()
+	FLiveLinkOpenTrackIOStaticData Static;
+
+	UPROPERTY()
+	FLiveLinkOpenTrackIOTracker Tracker;
+
+	UPROPERTY()
+	FLiveLinkOpenTrackIOTiming Timing;
+
+	UPROPERTY()
+	FLiveLinkOpenTrackIOLens Lens;
+
+	UPROPERTY()
+	FLiveLinkOpenTrackIOProtocol Protocol;
+
+	/**
+	 * URN serving as unique identifier of the sample in which data is being transported.
+	 * Pattern -> ^urn:uuid:[0-9a-f]{8}-[0-9a-f]{4}-[0-9a-f]{4}-[0-9a-f]{4}-[0-9a-f]{12}$
+	 */
+	UPROPERTY()
+	FString SampleId;
+
+	/**
+	 * URN serving as unique identifier of the source from which data is being transported.
+	 * Pattern -> ^urn:uuid:[0-9a-f]{8}-[0-9a-f]{4}-[0-9a-f]{4}-[0-9a-f]{4}-[0-9a-f]{12}$
+	 */
+	UPROPERTY()
+	FString SourceId;
+
+	/**
+	 * Number that identifies the index of the stream from a source from which
+	 * data is being transported. This is most important in the case where a source
+	 * is producing multiple streams of samples.
+	 */
+	UPROPERTY()
+	uint64 SourceNumber = 0;
+
+	/**
+	 * List of sampleId properties of samples related to this sample. The
+	 * existence of a sample with a given sampleId is not guaranteed.
+	 */
+	UPROPERTY()
+	TArray<FString> RelatedSampleIds;
+
+	/** Position of stage origin in global ENU and geodetic coordinates */
+	UPROPERTY()
+	FLiveLinkOpenTrackIOGlobalStage GlobalStage;
+
+	 /**
+	  * A list of transforms. Transforms are composed in order with the last in the list representing the X,Y,Z in meters of camera sensor relative to
+	  * stage origin. The Z axis points upwards and the coordinate system is right-handed. Y points in the forward camera direction (when pan, tilt
+	  * and roll are zero). For example in an LED volume Y would point towards the centre of the LED wall and so X would point to camera-right.
+	  * Rotation expressed as euler angles in degrees of the camera sensor relative to stage origin Rotations are intrinsic and are measured around
+	  * the axes ZXY, commonly referred to as [pan, tilt, roll] Notes on Euler angles: Euler angles are human readable and unlike quarternions,
+	  * provide the ability for cycles (with angles >360 or <0 degrees). Where a tracking system is providing the pose of a virtual camera, gimbal
+	  * lock does not present the physical challenges of a robotic system. Conversion to and from quarternions is trivial with an acceptable loss of
+	  * precision.
+	  */ 
+	UPROPERTY()
+	TArray<FLiveLinkOpenTrackIOTransform> Transforms;
+	
+};
+
diff --git a/Engine/Plugins/VirtualProduction/LiveLinkOpenTrackIO/Source/LiveLinkOpenTrackIO/Private/SLiveLinkOpenTrackIOSourceFactory.cpp b/Engine/Plugins/VirtualProduction/LiveLinkOpenTrackIO/Source/LiveLinkOpenTrackIO/Private/SLiveLinkOpenTrackIOSourceFactory.cpp
new file mode 100644
--- /dev/null
+++ b/Engine/Plugins/VirtualProduction/LiveLinkOpenTrackIO/Source/LiveLinkOpenTrackIO/Private/SLiveLinkOpenTrackIOSourceFactory.cpp
@@ -0,0 +1,62 @@
+// Copyright Epic Games, Inc. All Rights Reserved.
+
+#include "SLiveLinkOpenTrackIOSourceFactory.h"
+#include "LiveLinkOpenTrackIO.h"
+#include "LiveLinkOpenTrackIOSourceSettings.h"
+
+#include "Widgets/Input/SButton.h"
+#include "Widgets/SBoxPanel.h"
+
+#if WITH_EDITOR
+#include "DetailLayoutBuilder.h"
+#endif //WITH_EDITOR
+
+#define LOCTEXT_NAMESPACE "SLiveLinkOpenTrackIOSourceFactory"
+
+void SLiveLinkOpenTrackIOSourceFactory::Construct(const FArguments& Args)
+{
+#if WITH_EDITOR
+	OnConnectionSettingsAccepted = Args._OnConnectionSettingsAccepted;
+
+	FStructureDetailsViewArgs StructureViewArgs;
+	FDetailsViewArgs DetailArgs;
+	DetailArgs.bAllowSearch = false;
+	DetailArgs.bShowScrollBar = false;
+
+	FPropertyEditorModule& PropertyEditor = FModuleManager::Get().LoadModuleChecked<FPropertyEditorModule>(TEXT("PropertyEditor"));
+
+	StructOnScope = MakeShared<FStructOnScope>(FLiveLinkOpenTrackIOConnectionSettings::StaticStruct());
+	CastChecked<UScriptStruct>(StructOnScope->GetStruct())->CopyScriptStruct(StructOnScope->GetStructMemory(), &ConnectionSettings);
+	StructureDetailsView = PropertyEditor.CreateStructureDetailView(DetailArgs, StructureViewArgs, StructOnScope);
+
+	ChildSlot
+	[
+		SNew(SVerticalBox)
+		+SVerticalBox::Slot()
+		.FillHeight(1.f)
+		[
+			StructureDetailsView->GetWidget().ToSharedRef()
+		]
+		+ SVerticalBox::Slot()
+		.HAlign(HAlign_Right)
+		.AutoHeight()
+		[
+			SNew(SButton)
+			.OnClicked(this, &SLiveLinkOpenTrackIOSourceFactory::OnSettingsAccepted)
+			.Text(LOCTEXT("AddSource", "Add"))
+		]
+	];
+#endif //WITH_EDITOR
+}
+
+FReply SLiveLinkOpenTrackIOSourceFactory::OnSettingsAccepted()
+{
+#if WITH_EDITOR
+	CastChecked<UScriptStruct>(StructOnScope->GetStruct())->CopyScriptStruct(&ConnectionSettings, StructOnScope->GetStructMemory());
+	OnConnectionSettingsAccepted.ExecuteIfBound(ConnectionSettings);
+#endif //WITH_EDITOR
+
+	return FReply::Handled();
+}
+
+#undef LOCTEXT_NAMESPACE
\ No newline at end of file
diff --git a/Engine/Plugins/VirtualProduction/LiveLinkOpenTrackIO/Source/LiveLinkOpenTrackIO/Private/SLiveLinkOpenTrackIOSourceFactory.h b/Engine/Plugins/VirtualProduction/LiveLinkOpenTrackIO/Source/LiveLinkOpenTrackIO/Private/SLiveLinkOpenTrackIOSourceFactory.h
new file mode 100644
--- /dev/null
+++ b/Engine/Plugins/VirtualProduction/LiveLinkOpenTrackIO/Source/LiveLinkOpenTrackIO/Private/SLiveLinkOpenTrackIOSourceFactory.h
@@ -0,0 +1,39 @@
+// Copyright Epic Games, Inc. All Rights Reserved.
+
+#pragma once
+
+#include "LiveLinkOpenTrackIOConnectionSettings.h"
+#include "Widgets/SCompoundWidget.h"
+#include "Widgets/DeclarativeSyntaxSupport.h"
+
+#if WITH_EDITOR
+#include "IStructureDetailsView.h"
+#endif //WITH_EDITOR
+
+#include "Input/Reply.h"
+
+struct FLiveLinkOpenTrackIOConnectionSettings;
+
+DECLARE_DELEGATE_OneParam(FOnLiveLinkOpenTrackIOConnectionSettingsAccepted, FLiveLinkOpenTrackIOConnectionSettings);
+
+class SLiveLinkOpenTrackIOSourceFactory : public SCompoundWidget
+{
+	SLATE_BEGIN_ARGS(SLiveLinkOpenTrackIOSourceFactory)
+	{}
+		SLATE_EVENT(FOnLiveLinkOpenTrackIOConnectionSettingsAccepted, OnConnectionSettingsAccepted)
+	SLATE_END_ARGS()
+
+	void Construct(const FArguments& Args);
+
+
+private:
+	FLiveLinkOpenTrackIOConnectionSettings ConnectionSettings;
+
+#if WITH_EDITOR
+	TSharedPtr<FStructOnScope> StructOnScope;
+	TSharedPtr<IStructureDetailsView> StructureDetailsView;
+#endif //WITH_EDITOR
+
+	FReply OnSettingsAccepted();
+	FOnLiveLinkOpenTrackIOConnectionSettingsAccepted OnConnectionSettingsAccepted;
+};
\ No newline at end of file
diff --git a/Engine/Plugins/VirtualProduction/LiveLinkOpenTrackIO/Source/LiveLinkOpenTrackIO/Private/Tests/OpenTrackIOTestHelpers.cpp b/Engine/Plugins/VirtualProduction/LiveLinkOpenTrackIO/Source/LiveLinkOpenTrackIO/Private/Tests/OpenTrackIOTestHelpers.cpp
new file mode 100644
--- /dev/null
+++ b/Engine/Plugins/VirtualProduction/LiveLinkOpenTrackIO/Source/LiveLinkOpenTrackIO/Private/Tests/OpenTrackIOTestHelpers.cpp
@@ -0,0 +1,25 @@
+// Copyright Epic Games, Inc. All Rights Reserved.
+// 
+#include "OpenTrackIOTestHelpers.h"
+
+#include "Interfaces/IPluginManager.h"
+#include "Misc/FileHelper.h"
+#include "Misc/Paths.h"
+
+namespace UE::OpenTrackIO::Tests
+{
+
+FString GetSampleFile(const FString& InName)
+{
+	TSharedPtr<IPlugin> Plugin = IPluginManager::Get().FindPlugin(UE_PLUGIN_NAME);
+	check(Plugin.IsValid());
+
+	FString FilePath = FPaths::Combine(Plugin->GetBaseDir(),
+		TEXT("Source"), TEXT("LiveLinkOpenTrackIO"),
+		TEXT("Private"), TEXT("Tests"), TEXT("Samples"), InName);
+	ensure(FPaths::FileExists(FilePath));
+	return FilePath;
+}	
+
+}
+
diff --git a/Engine/Plugins/VirtualProduction/LiveLinkOpenTrackIO/Source/LiveLinkOpenTrackIO/Private/Tests/OpenTrackIOTestHelpers.h b/Engine/Plugins/VirtualProduction/LiveLinkOpenTrackIO/Source/LiveLinkOpenTrackIO/Private/Tests/OpenTrackIOTestHelpers.h
new file mode 100644
--- /dev/null
+++ b/Engine/Plugins/VirtualProduction/LiveLinkOpenTrackIO/Source/LiveLinkOpenTrackIO/Private/Tests/OpenTrackIOTestHelpers.h
@@ -0,0 +1,24 @@
+// Copyright Epic Games, Inc. All Rights Reserved.
+
+#pragma once
+
+#include "Containers/UnrealString.h"
+
+namespace UE::OpenTrackIO::Tests
+{
+FString GetSampleFile(const FString& InName);
+}
+
+
+
+
+
+
+
+
+
+
+
+
+
+
diff --git a/Engine/Plugins/VirtualProduction/LiveLinkOpenTrackIO/Source/LiveLinkOpenTrackIO/Private/Tests/OpenTrackIO_StaticExample.json b/Engine/Plugins/VirtualProduction/LiveLinkOpenTrackIO/Source/LiveLinkOpenTrackIO/Private/Tests/OpenTrackIO_StaticExample.json
new file mode 100644
--- /dev/null
+++ b/Engine/Plugins/VirtualProduction/LiveLinkOpenTrackIO/Source/LiveLinkOpenTrackIO/Private/Tests/OpenTrackIO_StaticExample.json
@@ -0,0 +1,93 @@
+{
+  "static": {
+    "camera": {
+      "activeSensorPhysicalDimensions": {
+        "height": 24.0,
+        "width": 36.0
+      },
+      "label": "A"
+    },
+    "lens": {
+      "make": "LensMaker",
+      "model": "Model15"
+    }
+  },
+  "tracker": {
+    "notes": "Example generated sample.",
+    "recording": false,
+    "slate": "A101_A_4",
+    "status": "Optical Good"
+  },
+  "timing": {
+    "mode": "external",
+    "sampleRate": {
+      "num": 24000,
+      "denom": 1001
+    },
+    "timecode": {
+      "hours": 1,
+      "minutes": 2,
+      "seconds": 3,
+      "frames": 4,
+      "frameRate": {
+        "num": 24000,
+        "denom": 1001
+      }
+    }
+  },
+  "lens": {
+    "distortion": [
+      {
+        "radial": [
+          1.0,
+          2.0,
+          3.0
+        ],
+        "tangential": [
+          1.0,
+          2.0
+        ],
+        "overscan": 3.1
+      }
+    ],
+    "encoders": {
+      "focus": 0.1,
+      "iris": 0.2,
+      "zoom": 0.3
+    },
+    "entrancePupilOffset": 0.123,
+    "fStop": 4.0,
+    "pinholeFocalLength": 24.305,
+    "focusDistance": 10.0,
+    "projectionOffset": {
+      "x": 0.1,
+      "y": 0.2
+    }
+  },
+  "protocol": {
+    "name": "OpenTrackIO",
+    "version": [
+      1,
+      0,
+      0
+    ]
+  },
+  "sampleId": "urn:uuid:c91a2031-1ebc-4e1d-af4b-a892208d5d9b",
+  "sourceId": "urn:uuid:116c7a5b-a21e-4ad1-bf02-05610e173c0d",
+  "sourceNumber": 1,
+  "transforms": [
+    {
+      "translation": {
+        "x": 1.0,
+        "y": 2.0,
+        "z": 3.0
+      },
+      "rotation": {
+        "pan": 180.0,
+        "tilt": 90.0,
+        "roll": 45.0
+      },
+      "id": "Camera"
+    }
+  ]
+}
diff --git a/Engine/Plugins/VirtualProduction/LiveLinkOpenTrackIO/Source/LiveLinkOpenTrackIO/Private/Tests/OpenTrackReaderTests.spec.cpp b/Engine/Plugins/VirtualProduction/LiveLinkOpenTrackIO/Source/LiveLinkOpenTrackIO/Private/Tests/OpenTrackReaderTests.spec.cpp
new file mode 100644
--- /dev/null
+++ b/Engine/Plugins/VirtualProduction/LiveLinkOpenTrackIO/Source/LiveLinkOpenTrackIO/Private/Tests/OpenTrackReaderTests.spec.cpp
@@ -0,0 +1,200 @@
+// Copyright Epic Games, Inc. All Rights Reserved.
+
+#include "OpenTrackIOTestHelpers.h"
+
+#if WITH_EDITOR
+
+#include "Misc/AutomationTest.h"
+#include "Misc/FileHelper.h"
+
+#include "LiveLinkOpenTrackIOParser.h"
+#include "LiveLinkOpenTrackIOTypes.h"
+	#include "LiveLinkOpenTrackIODatagram.h"
+
+#if WITH_DEV_AUTOMATION_TESTS
+
+namespace UE::OpenTrackIO::Tests
+{
+
+BEGIN_DEFINE_SPEC(FOpenTrackIOReaderTests,
+	TEXT("Editor.LiveLinkOpenTrackIO.ReaderTests"),
+	EAutomationTestFlags::EditorContext | EAutomationTestFlags::EngineFilter)
+	const TArray<FString> Tests = {
+		TEXT("FullStaticOpenTrack"),
+		TEXT("RecommendedDynamicExample"),
+		TEXT("CompleteStaticExample")
+	};
+
+	struct FOTIO_DataValidation
+	{
+		uint16 SequenceNumber;
+		uint16 Checksum;
+	};
+
+	const TArray<FOTIO_DataValidation> CborValidationData = []() -> TArray<FOTIO_DataValidation> {
+		return { {.SequenceNumber =  2,
+					 .Checksum = 0x4532},
+				 {.SequenceNumber =  3,
+					 .Checksum = 0xDFC0},
+				 {.SequenceNumber =  4,
+				     .Checksum = 0xDCA4},
+				 {.SequenceNumber =  5,
+				     .Checksum = 0xCD54},
+				 {.SequenceNumber =  6,
+				     .Checksum = 0xCEB},
+				 {.SequenceNumber =  7,
+				     .Checksum = 0x647},
+				 {.SequenceNumber =  8,
+				     .Checksum = 0xE29D}};
+	}();
+	const TArray<FOTIO_DataValidation> JsonValidationData = []() -> TArray<FOTIO_DataValidation> {
+		return { 
+			{.SequenceNumber = 2,
+			 .Checksum =  0xB728},
+			{.SequenceNumber = 3,
+			 .Checksum =  0x8308},
+			{.SequenceNumber = 4,
+			 .Checksum =  0x239},
+			{.SequenceNumber = 5,
+			 .Checksum =  0xCE19},
+			{.SequenceNumber = 6,
+			 .Checksum =  0x4D4A},
+			{.SequenceNumber = 7,
+			 .Checksum =  0x192A},
+			{.SequenceNumber = 8,
+			 .Checksum =  0x985B},
+			{.SequenceNumber = 9,
+			 .Checksum =  0x643B},
+			{.SequenceNumber = 10,
+			 .Checksum =  0xE36C},
+			{.SequenceNumber = 11,
+			.Checksum =  0x9793}};
+	}();
+END_DEFINE_SPEC(FOpenTrackIOReaderTests)
+
+void FOpenTrackIOReaderTests::Define()
+{
+	using namespace UE::OpenTrackIO::Tests;
+
+	BeforeEach([this] {
+		// Nothing right now. 
+	});
+	Describe("CanonicalCases", [this] {
+		It("Parses JSON", [this] {	
+			for (const FString& TestName : Tests)
+			{
+				FString JsonTestName = TestName + TEXT(".json");
+				FString FullPath = GetSampleFile(TestName + TEXT(".json"));
+
+				FString JsonBlob;
+				if (TestTrue("Parsed JSON -> " + JsonTestName, FFileHelper::LoadFileToString(JsonBlob, *FullPath)))
+				{
+					TOptional<FLiveLinkOpenTrackIOData> Data = UE::OpenTrackIO::Private::ParseJsonBlob(JsonBlob);
+					TestTrue(JsonTestName + " JSON should be successful.", Data.IsSet());
+				}
+			}
+		});
+
+		It("Fails with invalid JSON",[this]
+		{
+			FString InvalidJson = TEXT("{\"not opentrack\" : {}}");
+			TOptional<FLiveLinkOpenTrackIOData> Data = UE::OpenTrackIO::Private::ParseJsonBlob(InvalidJson);
+			
+			TestTrue("Should have failed with bogus data", !Data.IsSet());
+		});
+		
+		It("Parses CBOR", [this]
+		{	
+			for (const FString& TestName : Tests)
+			{
+				FString CborTestName = TestName + TEXT(".cbor");
+				FString FullPath = GetSampleFile(TestName + TEXT(".cbor"));
+
+				TArray<uint8> BinaryBlob;
+				if (TestTrue("Parsed CBOR -> " + CborTestName, FFileHelper::LoadFileToArray(BinaryBlob, *FullPath)))
+				{
+					TOptional<FLiveLinkOpenTrackIOData> Data = UE::OpenTrackIO::Private::ParseCborBlob(BinaryBlob);
+					TestTrue(CborTestName + " CBOR should be successful.", Data.IsSet());
+				}
+			}
+		});
+	});
+
+	Describe("ReadPackets", [this]
+	{
+		It("Reads CBOR Packets", [this] {
+			const FString CborPackets = "otio_cbor.packets";
+			FString FullPath = GetSampleFile(CborPackets);
+
+			TArray<uint8> PacketsBlob;
+			TestTrue("Did load cbor packets blob.", FFileHelper::LoadFileToArray(PacketsBlob, *FullPath));
+
+			TArrayView<uint8> AllPacketsView(PacketsBlob);
+			uint64 Index = 0;
+			for (const FOTIO_DataValidation& Item : CborValidationData)
+			{
+				TArrayView<uint8> PacketView = AllPacketsView.Slice(Index, AllPacketsView.Num()-Index);
+				
+				TOptional<FOpenTrackIOHeaderWithData> ParsedPayload = UE::OpenTrackIO::Private::ParseOpenTrackIOPacket(PacketView);
+				FString ParsedPacket = TEXT("Cbor Parsed Packet No ") + FString::FromInt(Item.SequenceNumber); 
+				TestTrue(ParsedPacket, ParsedPayload.IsSet());
+
+				if (ParsedPayload)
+				{
+					const FLiveLinkOpenTrackIODatagramHeader& Header = ParsedPayload->Get<0>();
+					const FLiveLinkOpenTrackIOData& Data = ParsedPayload->Get<1>();
+
+					TestTrue(ParsedPacket + " Has Payload", Header.GetPayloadSize() > 0);
+					TestTrue(ParsedPacket + " Matches Checksum", Item.Checksum == Header.Checksum);
+					TestTrue(ParsedPacket + " Matches Sequence", Item.SequenceNumber == Header.SequenceNumber);
+					Index += Header.GetPayloadSize() + sizeof(FLiveLinkOpenTrackIODatagramHeader); 
+				}
+				else
+				{
+					break;
+				}
+				
+			}
+		});
+		It("Reads JSON Packets", [this] {
+			const FString JsonPackets = "otio_json.packets";
+			FString FullPath = GetSampleFile(JsonPackets);
+
+			TArray<uint8> PacketsBlob;
+			TestTrue("Did load packets blob.", FFileHelper::LoadFileToArray(PacketsBlob, *FullPath));
+
+			TArrayView<uint8> AllPacketsView(PacketsBlob);
+			uint64 Index = 0;
+			for (const FOTIO_DataValidation& Item : JsonValidationData)
+			{
+				TArrayView<uint8> PacketView = AllPacketsView.Slice(Index, AllPacketsView.Num()-Index);
+				
+				TOptional<FOpenTrackIOHeaderWithData> ParsedPayload = UE::OpenTrackIO::Private::ParseOpenTrackIOPacket(PacketView);
+				FString ParsedPacket = TEXT("JSON Parsed Packet No ") + FString::FromInt(Item.SequenceNumber); 
+				TestTrue(ParsedPacket, ParsedPayload.IsSet());
+
+				if (ParsedPayload)
+				{
+					const FLiveLinkOpenTrackIODatagramHeader& Header = ParsedPayload->Get<0>();
+					const FLiveLinkOpenTrackIOData& Data = ParsedPayload->Get<1>();
+
+					TestTrue(ParsedPacket + " Has Payload", Header.GetPayloadSize() > 0);
+					TestTrue(ParsedPacket + " Matches Checksum", Item.Checksum == Header.Checksum);
+					TestTrue(ParsedPacket + " Matches Sequence", Item.SequenceNumber == Header.SequenceNumber);
+					Index += Header.GetPayloadSize() + sizeof(FLiveLinkOpenTrackIODatagramHeader); 
+				}
+				else
+				{
+					break;
+				}
+				
+			}
+		});
+	});
+}
+	
+}
+
+#endif // WITH_DEV_AUTOMATION_TESTS
+
+#endif // WITH_EDITOR
diff --git a/Engine/Plugins/VirtualProduction/LiveLinkOpenTrackIO/Source/LiveLinkOpenTrackIO/Private/Tests/Samples/CompleteStaticExample.json b/Engine/Plugins/VirtualProduction/LiveLinkOpenTrackIO/Source/LiveLinkOpenTrackIO/Private/Tests/Samples/CompleteStaticExample.json
new file mode 100644
--- /dev/null
+++ b/Engine/Plugins/VirtualProduction/LiveLinkOpenTrackIO/Source/LiveLinkOpenTrackIO/Private/Tests/Samples/CompleteStaticExample.json
@@ -0,0 +1,247 @@
+{
+  "static": {
+    "duration": {
+      "num": 1,
+      "denom": 25
+    },
+    "camera": {
+      "captureFrameRate": {
+        "num": 24000,
+        "denom": 1001
+      },
+      "activeSensorPhysicalDimensions": {
+        "height": 24.0,
+        "width": 36.0
+      },
+      "activeSensorResolution": {
+        "height": 2160,
+        "width": 3840
+      },
+      "make": "CameraMaker",
+      "model": "Model20",
+      "serialNumber": "1234567890A",
+      "firmwareVersion": "1.2.3",
+      "label": "A",
+      "anamorphicSqueeze": {
+        "num": 1,
+        "denom": 1
+      },
+      "isoSpeed": 4000,
+      "fdlLink": "urn:uuid:e386a0b9-4f20-4767-88bd-8b6066a78123",
+      "shutterAngle": 45.0
+    },
+    "lens": {
+      "distortionOverscanMax": 1.2,
+      "undistortionOverscanMax": 1.3,
+      "make": "LensMaker",
+      "model": "Model15",
+      "serialNumber": "1234567890A",
+      "nominalFocalLength": 14.0,
+      "calibrationHistory": [
+        "LensMaker 123",
+        "TrackerMaker 123"
+      ]
+    },
+    "tracker": {
+      "make": "TrackerMaker",
+      "model": "Tracker",
+      "serialNumber": "1234567890A",
+      "firmwareVersion": "1.2.3"
+    }
+  },
+  "tracker": {
+    "notes": "Example generated sample.",
+    "recording": false,
+    "slate": "A101_A_4",
+    "status": "Optical Good"
+  },
+  "timing": {
+    "mode": "internal",
+    "recordedTimestamp": {
+      "seconds": 1718806000,
+      "nanoseconds": 500000000
+    },
+    "sampleRate": {
+      "num": 24000,
+      "denom": 1001
+    },
+    "sampleTimestamp": {
+      "seconds": 1718806554,
+      "nanoseconds": 500000000
+    },
+    "sequenceNumber": 0,
+    "synchronization": {
+      "locked": true,
+      "source": "ptp",
+      "frequency": {
+        "num": 24000,
+        "denom": 1001
+      },
+      "present": true,
+      "ptp": {
+        "profile": "SMPTE ST2059-2:2021",
+        "domain": 1,
+        "leaderIdentity": "00:11:22:33:44:55",
+        "leaderPriorities": {
+          "priority1": 128,
+          "priority2": 128
+        },
+        "leaderAccuracy": 5e-08,
+        "leaderTimeSource": "GNSS",
+        "meanPathDelay": 0.000123,
+        "vlan": 100
+      }
+    },
+    "timecode": {
+      "hours": 1,
+      "minutes": 2,
+      "seconds": 3,
+      "frames": 4,
+      "frameRate": {
+        "num": 24000,
+        "denom": 1001
+      }
+    }
+  },
+  "lens": {
+    "custom": [
+      1.0,
+      2.0
+    ],
+    "distortion": [
+      {
+        "model": "Brown-Conrady U-D",
+        "radial": [
+          1.0,
+          2.0,
+          3.0,
+          4.0,
+          5.0,
+          6.0
+        ],
+        "tangential": [
+          1.0,
+          2.0
+        ],
+        "overscan": 3.0
+      },
+      {
+        "radial": [
+          1.0,
+          2.0,
+          3.0,
+          4.0,
+          5.0,
+          6.0
+        ],
+        "tangential": [
+          1.0,
+          2.0
+        ],
+        "overscan": 2.0
+      }
+    ],
+    "distortionOffset": {
+      "x": 1.0,
+      "y": 2.0
+    },
+    "encoders": {
+      "focus": 0.1,
+      "iris": 0.2,
+      "zoom": 0.3
+    },
+    "entrancePupilOffset": 0.123,
+    "exposureFalloff": {
+      "a1": 1.0,
+      "a2": 2.0,
+      "a3": 3.0
+    },
+    "fStop": 4.0,
+    "pinholeFocalLength": 24.305,
+    "focusDistance": 10.0,
+    "projectionOffset": {
+      "x": 0.1,
+      "y": 0.2
+    },
+    "rawEncoders": {
+      "focus": 1000,
+      "iris": 2000,
+      "zoom": 3000
+    },
+    "tStop": 4.1
+  },
+  "protocol": {
+    "name": "OpenTrackIO",
+    "version": [
+      1,
+      0,
+      0
+    ]
+  },
+  "sampleId": "urn:uuid:939fb522-68f1-40cc-a971-b7a5b7aededa",
+  "sourceId": "urn:uuid:1530983c-4fbf-4cf5-b054-e6848365ca67",
+  "sourceNumber": 1,
+  "relatedSampleIds": [
+    "urn:uuid:fb7f644e-29f9-4c8d-910e-9b26a2639863",
+    "urn:uuid:3384c500-e45c-4bc2-ab67-4e18efad7152"
+  ],
+  "globalStage": {
+    "E": 100.0,
+    "N": 200.0,
+    "U": 300.0,
+    "lat0": 100.0,
+    "lon0": 200.0,
+    "h0": 300.0
+  },
+  "transforms": [
+    {
+      "translation": {
+        "x": 1.0,
+        "y": 2.0,
+        "z": 3.0
+      },
+      "rotation": {
+        "pan": 180.0,
+        "tilt": 90.0,
+        "roll": 45.0
+      },
+      "id": "Dolly"
+    },
+    {
+      "translation": {
+        "x": 1.0,
+        "y": 2.0,
+        "z": 3.0
+      },
+      "rotation": {
+        "pan": 180.0,
+        "tilt": 90.0,
+        "roll": 45.0
+      },
+      "scale": {
+        "x": 1.0,
+        "y": 2.0,
+        "z": 3.0
+      },
+      "id": "Crane Arm"
+    },
+    {
+      "translation": {
+        "x": 1.0,
+        "y": 2.0,
+        "z": 3.0
+      },
+      "rotation": {
+        "pan": 180.0,
+        "tilt": 90.0,
+        "roll": 45.0
+      },
+      "scale": {
+        "x": 1.0,
+        "y": 2.0,
+        "z": 3.0
+      },
+      "id": "Camera"
+    }
+  ]
+}
diff --git a/Engine/Plugins/VirtualProduction/LiveLinkOpenTrackIO/Source/LiveLinkOpenTrackIO/Private/Tests/Samples/FullStaticOpenTrack.json b/Engine/Plugins/VirtualProduction/LiveLinkOpenTrackIO/Source/LiveLinkOpenTrackIO/Private/Tests/Samples/FullStaticOpenTrack.json
new file mode 100644
--- /dev/null
+++ b/Engine/Plugins/VirtualProduction/LiveLinkOpenTrackIO/Source/LiveLinkOpenTrackIO/Private/Tests/Samples/FullStaticOpenTrack.json
@@ -0,0 +1,96 @@
+{
+  "static": {
+    "camera": {
+      "activeSensorPhysicalDimensions": {
+        "height": 24.0,
+        "width": 36.0
+      },
+      "label": "A"
+    },
+    "lens": {
+      "make": "LensMaker",
+      "model": "Model15"
+    }
+  },
+  "tracker": {
+    "notes": "Example generated sample.",
+    "recording": false,
+    "slate": "A101_A_4",
+    "status": "Optical Good"
+  },
+  "timing": {
+    "mode": "external",
+    "sampleRate": {
+      "num": 24000,
+      "denom": 1001
+    },
+    "timecode": {
+      "hours": 1,
+      "minutes": 2,
+      "seconds": 3,
+      "frames": 4,
+      "frameRate": {
+        "num": 24000,
+        "denom": 1001
+      }
+    }
+  },
+  "lens": {
+    "distortion": [
+      {
+        "radial": [
+          1.0,
+          2.0,
+          3.0
+        ],
+        "tangential": [
+          1.0,
+          2.0
+        ],
+        "overscan": 3.1
+      }
+    ],
+    "encoders": {
+      "focus": 0.1,
+      "iris": 0.2,
+      "zoom": 0.3
+    },
+    "entrancePupilOffset": 0.123,
+    "fStop": 4.0,
+    "pinholeFocalLength": 24.305,
+    "focusDistance": 10.0,
+    "projectionOffset": {
+      "x": 0.1,
+      "y": 0.2
+    }
+  },
+  "protocol": {
+    "name": "OpenTrackIO",
+    "version": [
+      1,
+      0,
+      0
+    ]
+  },
+  "sampleId": "urn:uuid:c91a2031-1ebc-4e1d-af4b-a892208d5d9b",
+  "sourceId": "urn:uuid:116c7a5b-a21e-4ad1-bf02-05610e173c0d",
+  "sourceNumber": 1,
+  "transforms": [
+    {
+      "translation": {
+        "x": 1.0,
+        "y": 2.0,
+        "z": 3.0
+      },
+      "rotation": {
+        "pan": 180.0,
+        "tilt": 90.0,
+        "roll": 45.0
+      },
+      "id": "Camera"
+    }
+  ]
+}
+
+
+
diff --git a/Engine/Plugins/VirtualProduction/LiveLinkOpenTrackIO/Source/LiveLinkOpenTrackIO/Private/Tests/Samples/RecommendedDynamicExample.json b/Engine/Plugins/VirtualProduction/LiveLinkOpenTrackIO/Source/LiveLinkOpenTrackIO/Private/Tests/Samples/RecommendedDynamicExample.json
new file mode 100644
--- /dev/null
+++ b/Engine/Plugins/VirtualProduction/LiveLinkOpenTrackIO/Source/LiveLinkOpenTrackIO/Private/Tests/Samples/RecommendedDynamicExample.json
@@ -0,0 +1,80 @@
+{
+  "tracker": {
+    "notes": "Example generated sample.",
+    "recording": false,
+    "slate": "A101_A_4",
+    "status": "Optical Good"
+  },
+  "timing": {
+    "mode": "external",
+    "sampleRate": {
+      "num": 24000,
+      "denom": 1001
+    },
+    "timecode": {
+      "hours": 1,
+      "minutes": 2,
+      "seconds": 3,
+      "frames": 4,
+      "frameRate": {
+        "num": 24000,
+        "denom": 1001
+      }
+    }
+  },
+  "lens": {
+    "distortion": [
+      {
+        "radial": [
+          1.0,
+          2.0,
+          3.0
+        ],
+        "tangential": [
+          1.0,
+          2.0
+        ],
+        "overscan": 3.1
+      }
+    ],
+    "encoders": {
+      "focus": 0.1,
+      "iris": 0.2,
+      "zoom": 0.3
+    },
+    "entrancePupilOffset": 0.123,
+    "fStop": 4.0,
+    "pinholeFocalLength": 24.305,
+    "focusDistance": 10.0,
+    "projectionOffset": {
+      "x": 0.1,
+      "y": 0.2
+    }
+  },
+  "protocol": {
+    "name": "OpenTrackIO",
+    "version": [
+      1,
+      0,
+      0
+    ]
+  },
+  "sampleId": "urn:uuid:771aac95-fa2d-49a8-9fdb-1e07afb6ffa0",
+  "sourceId": "urn:uuid:6dcc8625-645e-404f-8ed4-0edb18055dd1",
+  "sourceNumber": 1,
+  "transforms": [
+    {
+      "translation": {
+        "x": 1.0,
+        "y": 2.0,
+        "z": 3.0
+      },
+      "rotation": {
+        "pan": 180.0,
+        "tilt": 90.0,
+        "roll": 45.0
+      },
+      "id": "Camera"
+    }
+  ]
+}
\ No newline at end of file
diff --git a/Engine/Plugins/VirtualProduction/LiveLinkOpenTrackIO/Source/LiveLinkOpenTrackIO/Public/LiveLinkOpenTrackIO.h b/Engine/Plugins/VirtualProduction/LiveLinkOpenTrackIO/Source/LiveLinkOpenTrackIO/Public/LiveLinkOpenTrackIO.h
new file mode 100644
--- /dev/null
+++ b/Engine/Plugins/VirtualProduction/LiveLinkOpenTrackIO/Source/LiveLinkOpenTrackIO/Public/LiveLinkOpenTrackIO.h
@@ -0,0 +1,24 @@
+// Copyright Epic Games, Inc. All Rights Reserved.
+
+#pragma once
+
+#include "Modules/ModuleManager.h"
+
+DECLARE_LOG_CATEGORY_EXTERN(LogLiveLinkOpenTrackIO, Log, All);
+
+/**
+ * This module reads OpenTrackIO data from a UDP source.
+ */
+class FLiveLinkOpenTrackIOModule : public IModuleInterface
+{
+public:
+
+	static FLiveLinkOpenTrackIOModule& Get()
+	{
+		return FModuleManager::Get().LoadModuleChecked<FLiveLinkOpenTrackIOModule>(TEXT("LiveLinkOpenTrackIO"));
+	}
+
+	/** IModuleInterface implementation */
+	virtual void StartupModule() override;
+	virtual void ShutdownModule() override;
+};
diff --git a/Engine/Plugins/VirtualProduction/LiveLinkOpenTrackIO/Source/LiveLinkOpenTrackIO/Public/LiveLinkOpenTrackIOSourceSettings.h b/Engine/Plugins/VirtualProduction/LiveLinkOpenTrackIO/Source/LiveLinkOpenTrackIO/Public/LiveLinkOpenTrackIOSourceSettings.h
new file mode 100644
--- /dev/null
+++ b/Engine/Plugins/VirtualProduction/LiveLinkOpenTrackIO/Source/LiveLinkOpenTrackIO/Public/LiveLinkOpenTrackIOSourceSettings.h
@@ -0,0 +1,48 @@
+// Copyright Epic Games, Inc. All Rights Reserved.
+
+#pragma once
+
+#include "LiveLinkSourceSettings.h"
+#include "LiveLinkOpenTrackIOSourceSettings.generated.h"
+
+UCLASS()
+class LIVELINKOPENTRACKIO_API ULiveLinkOpenTrackIOSourceSettings : public ULiveLinkSourceSettings
+{
+	GENERATED_BODY()
+
+public:
+	/**
+	 * The IP endpoint to send multicast packets to.
+	 *
+	 * The format is IP_ADDRESS:PORT_NUMBER.
+	 * The multicast IP address must be in the range 224.0.0.0 to 239.255.255.255.
+	 *
+	 * The default value is 235.135.1.14:55555 per OpenTrackIO specifications.
+	 */
+	UPROPERTY(EditAnywhere, Category=Transport)
+	FString MulticastEndpoint = TEXT("235.135.1.14:55555");
+
+	/**
+	 * The IP endpoint to listen to and send packets from.  
+	 *
+	 * The format is IP_ADDRESS:PORT_NUMBER.
+	 * 0.0.0.0:0 will bind to the default network adapter on Windows,
+	 * and all available network adapters on other operating systems.
+	 *
+	 */
+	UPROPERTY(EditAnywhere, Category=Transport)
+	FString UnicastEndpoint = TEXT("0.0.0.0:0");
+
+	/**
+	 * The assigned IP address and port used by the socket layer.  This is for inspection of the connection
+	 * setup by the source via either blueprint or details panels for the source. 
+	 *
+	 * If this value starts with 0.0.0.0 then any adapter can be used with the given port number.
+	 * The format is IP_ADDRESS:PORT_NUMBER.
+	 *
+	 * If no receiver is setup then this value will be an empty string. 
+	 *
+	 */
+	UPROPERTY(VisibleAnywhere, Category=Transport, meta = (DisplayName = "Assigned Endpoint") )
+	FString AssignedUnicastEndpoint = TEXT("");
+};
