Remove unacceptable words from Live Link OpenTrackIO comment.

#rnx

diff --git a/Engine/Plugins/VirtualProduction/LiveLinkOpenTrackIO/Source/LiveLinkOpenTrackIO/Private/LiveLinkOpenTrackIOTypes.h b/Engine/Plugins/VirtualProduction/LiveLinkOpenTrackIO/Source/LiveLinkOpenTrackIO/Private/LiveLinkOpenTrackIOTypes.h
--- a/Engine/Plugins/VirtualProduction/LiveLinkOpenTrackIO/Source/LiveLinkOpenTrackIO/Private/LiveLinkOpenTrackIOTypes.h
+++ b/Engine/Plugins/VirtualProduction/LiveLinkOpenTrackIO/Source/LiveLinkOpenTrackIO/Private/LiveLinkOpenTrackIOTypes.h
@@ -568,7 +568,7 @@ struct FLiveLinkOpenTrackIO_PTP
  *
  *  - "domain": Identifies the PTP domain the device belongs to. Devices in the same domain can synchronize with each other
  *  
- *  - "leaderIdentity": The unique identifier (usually MAC address) of the current PTP leader (grandmaster)
+ *  - "leaderIdentity": The unique identifier (usually MAC address) of the current PTP leader
  *  
  *  - "leaderPriorities": The priority values of the leader used in the Best Clock Algorithm (BMCA). Lower values indicate higher priority
  *  
