Fix duplicate localization keys.

diff --git a/Engine/Plugins/VirtualProduction/LiveLinkOpenTrackIO/Source/LiveLinkOpenTrackIO/Private/LiveLinkOpenTrackIOSource.cpp b/Engine/Plugins/VirtualProduction/LiveLinkOpenTrackIO/Source/LiveLinkOpenTrackIO/Private/LiveLinkOpenTrackIOSource.cpp
--- a/Engine/Plugins/VirtualProduction/LiveLinkOpenTrackIO/Source/LiveLinkOpenTrackIO/Private/LiveLinkOpenTrackIOSource.cpp
+++ b/Engine/Plugins/VirtualProduction/LiveLinkOpenTrackIO/Source/LiveLinkOpenTrackIO/Private/LiveLinkOpenTrackIOSource.cpp
@@ -227,7 +227,7 @@ void FLiveLinkOpenTrackIOSource::InitializeSettings(ULiveLinkSourceSettings* Set
 
 void FLiveLinkOpenTrackIOSource::ResetSockets()
 {
-	SourceStatus = LOCTEXT("Resetting", "Resetting...");
+	SourceStatus = LOCTEXT("Resetting", "Resetting source.");
 	ConnectionState = ELiveLinkOpenTrackIOState::ResetRequested;
 	if (IsValid(SavedSourceSettings))
 	{
