Minor fixes from the initial submission of LiveLinkOpenTrackIO plugin

diff --git a/Engine/Plugins/VirtualProduction/LiveLinkOpenTrackIO/Source/LiveLinkOpenTrackIO/Private/LiveLinkOpenTrackIODatagram.h b/Engine/Plugins/VirtualProduction/LiveLinkOpenTrackIO/Source/LiveLinkOpenTrackIO/Private/LiveLinkOpenTrackIODatagram.h
--- a/Engine/Plugins/VirtualProduction/LiveLinkOpenTrackIO/Source/LiveLinkOpenTrackIO/Private/LiveLinkOpenTrackIODatagram.h
+++ b/Engine/Plugins/VirtualProduction/LiveLinkOpenTrackIO/Source/LiveLinkOpenTrackIO/Private/LiveLinkOpenTrackIODatagram.h
@@ -72,6 +72,8 @@ struct FLiveLinkOpenTrackIODatagramHeader
 	/** Helper function to return the size of the payload by and-ing the last 15 bits. */
 	uint16 GetPayloadSize() const
 	{
+		static_assert(sizeof(FLiveLinkOpenTrackIODatagramHeader) == 16, "OpenTrackIO header is expected to be 16 bytes.");
+		
 		return LastSegmentFlagAndPayloadLength & 0x7FFF;
 	}
 	
diff --git a/Engine/Plugins/VirtualProduction/LiveLinkOpenTrackIO/Source/LiveLinkOpenTrackIO/Private/LiveLinkOpenTrackIOParser.cpp b/Engine/Plugins/VirtualProduction/LiveLinkOpenTrackIO/Source/LiveLinkOpenTrackIO/Private/LiveLinkOpenTrackIOParser.cpp
--- a/Engine/Plugins/VirtualProduction/LiveLinkOpenTrackIO/Source/LiveLinkOpenTrackIO/Private/LiveLinkOpenTrackIOParser.cpp
+++ b/Engine/Plugins/VirtualProduction/LiveLinkOpenTrackIO/Source/LiveLinkOpenTrackIO/Private/LiveLinkOpenTrackIOParser.cpp
@@ -107,7 +107,7 @@ const uint32 OpenTrackIOHeaderId = 0x4F54726B;
 const uint8 JSON_Encoding = 0x01;
 const uint8 CBOR_Encoding = 0x02;
 
-uint16 ByteSwap(uint16 InVal)
+FORCEINLINE uint16 ByteSwap(uint16 InVal)
 {
 #if PLATFORM_LITTLE_ENDIAN
 	return BYTESWAP_ORDER16(InVal);
@@ -116,7 +116,7 @@ uint16 ByteSwap(uint16 InVal)
 #endif
 }
 
-uint32 ByteSwap(uint32 InVal)
+FORCEINLINE uint32 ByteSwap(uint32 InVal)
 {
 #if PLATFORM_LITTLE_ENDIAN
 	return BYTESWAP_ORDER32(InVal);
diff --git a/Engine/Plugins/VirtualProduction/LiveLinkOpenTrackIO/Source/LiveLinkOpenTrackIO/Private/LiveLinkOpenTrackIOParser.h b/Engine/Plugins/VirtualProduction/LiveLinkOpenTrackIO/Source/LiveLinkOpenTrackIO/Private/LiveLinkOpenTrackIOParser.h
--- a/Engine/Plugins/VirtualProduction/LiveLinkOpenTrackIO/Source/LiveLinkOpenTrackIO/Private/LiveLinkOpenTrackIOParser.h
+++ b/Engine/Plugins/VirtualProduction/LiveLinkOpenTrackIO/Source/LiveLinkOpenTrackIO/Private/LiveLinkOpenTrackIOParser.h
@@ -13,9 +13,10 @@ namespace UE::OpenTrackIO::Private
 {
 	/**
 	 * Try to parse a JSON string and return an optional FLiveLinkOpenTrackIOData type.  If parsing fails
-	 * a unset optional will be returned.  If parsing succeeds then data will be returned.
+	 * an unset optional will be returned.  If parsing succeeds then data will be returned.
 	 *
-	 * It is assumed that the JSON blob conforms to the minimal JSON 
+	 * It is assumed that the JSON blob conforms to the minimal JSON. If it does not conform to the
+	 * minimal JSON then an unset optional will be returned.
 	 */
 	TOptional<FLiveLinkOpenTrackIOData> ParseJsonBlob(const FString& JsonBlob);
 
diff --git a/Engine/Plugins/VirtualProduction/LiveLinkOpenTrackIO/Source/LiveLinkOpenTrackIO/Private/LiveLinkOpenTrackIOSource.cpp b/Engine/Plugins/VirtualProduction/LiveLinkOpenTrackIO/Source/LiveLinkOpenTrackIO/Private/LiveLinkOpenTrackIOSource.cpp
--- a/Engine/Plugins/VirtualProduction/LiveLinkOpenTrackIO/Source/LiveLinkOpenTrackIO/Private/LiveLinkOpenTrackIOSource.cpp
+++ b/Engine/Plugins/VirtualProduction/LiveLinkOpenTrackIO/Source/LiveLinkOpenTrackIO/Private/LiveLinkOpenTrackIOSource.cpp
@@ -42,9 +42,9 @@ TOptional<FIPv4Endpoint> ParseEndpoint(const FString& InEndpointString)
 	if (bParsedAddr)
 	{
 		// Detect 169.254.x.x addresses.
-		if (Endpoint.Address.A == 169 &&	Endpoint.Address.B == 254)
+		if (Endpoint.Address.A == 169 && Endpoint.Address.B == 254)
 		{
-			UE_LOG(LogLiveLinkOpenTrackIO, Warning, TEXT("Detected IPv4 address in the form of 169.254.x.x. This is a link assigned address and may prevent you from reach external endpoints. "));
+			UE_LOG(LogLiveLinkOpenTrackIO, Warning, TEXT("Detected IPv4 address in the form of 169.254.x.x. This is a link assigned address and may prevent you from reaching external endpoints. "));
 		}
 
 		return Endpoint;
@@ -382,7 +382,7 @@ bool FLiveLinkOpenTrackIOSource::OpenSockets()
 	}
 	else
 	{
-		int32 PortNo = UnicastSocket->GetPortNo();
+		const int32 PortNo = UnicastSocket->GetPortNo();
 		FString IpPortStr = FString::Format(TEXT("{0}:{1}"), {UnicastEndpoint.Address.ToString(), PortNo});
 		UE_LOG(LogLiveLinkOpenTrackIO, Display, TEXT("OpenTrackIO unicast socket socket bound to '%s'."), *IpPortStr);
 
diff --git a/Engine/Plugins/VirtualProduction/LiveLinkOpenTrackIO/Source/LiveLinkOpenTrackIO/Private/LiveLinkOpenTrackIOTypes.h b/Engine/Plugins/VirtualProduction/LiveLinkOpenTrackIO/Source/LiveLinkOpenTrackIO/Private/LiveLinkOpenTrackIOTypes.h
--- a/Engine/Plugins/VirtualProduction/LiveLinkOpenTrackIO/Source/LiveLinkOpenTrackIO/Private/LiveLinkOpenTrackIOTypes.h
+++ b/Engine/Plugins/VirtualProduction/LiveLinkOpenTrackIO/Source/LiveLinkOpenTrackIO/Private/LiveLinkOpenTrackIOTypes.h
@@ -180,7 +180,7 @@ struct FLiveLinkOpenTrackIOStaticCamera
 	UPROPERTY()
 	FLiveLinkOpenTrackIOFrameRate CaptureFrameRate;
 
-	/** Object representing the sensor dimenions. */
+	/** Object representing the sensor dimensions. */
 	UPROPERTY()
 	FLiveLinkOpenTrackIOSensorDimensions ActiveSensorPhysicalDimensions;
 
@@ -228,7 +228,7 @@ struct FLiveLinkOpenTrackIOStaticCamera
 	 * to the value of the parameter divided by 360 times the capture frame rate.
 	 */
 	UPROPERTY(meta = (ClampMin = "0.0", ClampMax = "360.0"))
-	float ShutterAngle = 0.0;
+	float ShutterAngle = 180.0;
 };
 
 USTRUCT()
