OpenTrackIO: Update Fletcher16 to use modulo 255.

This complies with the implementation in CamDKIt.
Also fixed a crash when removing the source from the Live Link panel due to the UdpSocketReceiver thread still running after the object was destroyed.

diff --git a/Engine/Plugins/VirtualProduction/LiveLinkOpenTrackIO/Source/LiveLinkOpenTrackIO/Private/LiveLinkOpenTrackIOParser.cpp b/Engine/Plugins/VirtualProduction/LiveLinkOpenTrackIO/Source/LiveLinkOpenTrackIO/Private/LiveLinkOpenTrackIOParser.cpp
--- a/Engine/Plugins/VirtualProduction/LiveLinkOpenTrackIO/Source/LiveLinkOpenTrackIO/Private/LiveLinkOpenTrackIOParser.cpp
+++ b/Engine/Plugins/VirtualProduction/LiveLinkOpenTrackIO/Source/LiveLinkOpenTrackIO/Private/LiveLinkOpenTrackIOParser.cpp
@@ -79,25 +79,36 @@ TOptional<FLiveLinkOpenTrackIOData> ParseCborBlob(TArrayView<uint8> InBytes)
 }
 
 /**
- * OpenTrackIO specifies using Fletcher-16 algorithm for checksum. This follows from the reference
- * code for Fletcher-16. The article on wikipedia suggests some possible performance improvements
- * but we opted for the unoptimized version: https://en.wikipedia.org/wiki/Fletcher%27s_checksum
+ * OpenTrackIO specifies using Fletcher-16 algorithm for checksum. 
+ * Optimized version from: https://en.wikipedia.org/wiki/Fletcher%27s_checksum
  */
 uint16_t CalculateChecksum(TArrayView<uint8> HeaderBytes, TArrayView<uint8> PayloadBytes)
 {
-	uint8 Sum1 = 0;
-	uint8 Sum2 = 0;
-	for (uint8 Byte : HeaderBytes)
-	{
-		Sum1 = (Sum1 + Byte) % 256;
-		Sum2 = (Sum2 + Sum1) % 256;
-	}
-	for (uint8 Byte : PayloadBytes)
-	{
-		Sum1 = (Sum1 + Byte) % 256;
-		Sum2 = (Sum2 + Sum1) % 256;		
-	}
-	return ((uint16) Sum2 << 8) | (uint16) Sum1;
+	uint32_t C0 = 0;
+	uint32_t C1 = 0;
+
+	auto CalcFletcher16Chunk = [&](const uint8* Data, uint32 Len)
+		{
+			while (Len > 0)
+			{
+				uint32 BlockLen = Len > 5802 ? 5802 : Len;
+
+				Len -= BlockLen;
+
+				do {
+					C0 += *Data++;
+					C1 += C0;
+				} while (--BlockLen);
+
+				C0 %= 255;
+				C1 %= 255;
+			}
+		};
+
+	CalcFletcher16Chunk(HeaderBytes.GetData(), HeaderBytes.Num());
+	CalcFletcher16Chunk(PayloadBytes.GetData(), PayloadBytes.Num());
+
+	return static_cast<uint16>((C1 << 8) | C0);
 }
 
 /**
diff --git a/Engine/Plugins/VirtualProduction/LiveLinkOpenTrackIO/Source/LiveLinkOpenTrackIO/Private/LiveLinkOpenTrackIOSource.cpp b/Engine/Plugins/VirtualProduction/LiveLinkOpenTrackIO/Source/LiveLinkOpenTrackIO/Private/LiveLinkOpenTrackIOSource.cpp
--- a/Engine/Plugins/VirtualProduction/LiveLinkOpenTrackIO/Source/LiveLinkOpenTrackIO/Private/LiveLinkOpenTrackIOSource.cpp
+++ b/Engine/Plugins/VirtualProduction/LiveLinkOpenTrackIO/Source/LiveLinkOpenTrackIO/Private/LiveLinkOpenTrackIOSource.cpp
@@ -339,6 +339,16 @@ void FLiveLinkOpenTrackIOSource::PushDataToLiveLink_AnyThread(const FLiveLinkOpe
 
 void FLiveLinkOpenTrackIOSource::CloseSockets()
 {
+	if (MulticastReceiver)
+	{
+		MulticastReceiver->Stop();
+	}
+	
+	if (UnicastReceiver)
+	{
+		UnicastReceiver->Stop();
+	}
+
 	MulticastReceiver.Reset();
 	UnicastReceiver.Reset();
 
