OpenTrackIO: Converted coordinates from OpenTrackIO to Unreal.

It also applies the entrance pupil offset as a final transform in the chain.

diff --git a/Engine/Plugins/VirtualProduction/LiveLinkOpenTrackIO/Source/LiveLinkOpenTrackIO/Private/LiveLinkOpenTrackIOTranscoder.cpp b/Engine/Plugins/VirtualProduction/LiveLinkOpenTrackIO/Source/LiveLinkOpenTrackIO/Private/LiveLinkOpenTrackIOTranscoder.cpp
--- a/Engine/Plugins/VirtualProduction/LiveLinkOpenTrackIO/Source/LiveLinkOpenTrackIO/Private/LiveLinkOpenTrackIOTranscoder.cpp
+++ b/Engine/Plugins/VirtualProduction/LiveLinkOpenTrackIO/Source/LiveLinkOpenTrackIO/Private/LiveLinkOpenTrackIOTranscoder.cpp
@@ -70,7 +70,20 @@ TArray<FSubjectAndFrameDataTuple> FLiveLinkOpenTrackIOCache::TryReadFrameData(co
 	FTransform FinalTransform; 
 	for (const FLiveLinkOpenTrackIOTransform& Transform : Data.Transforms)
 	{
-		FinalTransform *= Transform.GetTransform();
+		FinalTransform *= Transform.ToUnrealTransform();
+	}
+
+	// Apply entrance pupil offset as a final transform in the forward direction.
+	{
+		const FLiveLinkOpenTrackIO_XYZ EntrancePupilAsXYZ(0, Data.Lens.EntrancePupilOffset, 0);
+
+		const FTransform PupilOffset(
+			FQuat::Identity,
+			EntrancePupilAsXYZ.ToUnrealTranslation(),
+			FVector(1, 1, 1)
+		);
+
+		FinalTransform *= PupilOffset;
 	}
 
 	FQualifiedFrameTime SampleTC;
diff --git a/Engine/Plugins/VirtualProduction/LiveLinkOpenTrackIO/Source/LiveLinkOpenTrackIO/Private/LiveLinkOpenTrackIOTypes.h b/Engine/Plugins/VirtualProduction/LiveLinkOpenTrackIO/Source/LiveLinkOpenTrackIO/Private/LiveLinkOpenTrackIOTypes.h
--- a/Engine/Plugins/VirtualProduction/LiveLinkOpenTrackIO/Source/LiveLinkOpenTrackIO/Private/LiveLinkOpenTrackIOTypes.h
+++ b/Engine/Plugins/VirtualProduction/LiveLinkOpenTrackIO/Source/LiveLinkOpenTrackIO/Private/LiveLinkOpenTrackIOTypes.h
@@ -10,6 +10,31 @@
 
 #include "LiveLinkOpenTrackIOTypes.generated.h"
 
+
+namespace LiveLinkOpenTrackIO
+{
+	/** OpenTrackIO coordinate system 
+	 * 
+	 * OpenTrackIO is:
+	 * o Z-up
+	 * o Y-forward
+	 * o Right-handed
+	 *
+	 * Unreal is:
+	 * o Z-up
+	 * o X-forward
+	 * o Left-handed
+	 *
+	 * To convert translation and scale, swap X and Y, which both updates the forward vector and toggles the handedness.
+	 * Translation also needs a factor of 100 because OpenTrackIO units are meters, while Unreal uses centimeters.
+	 * To convert rotations, invert the sign of Pan/Yaw. The others can stay the same.
+	 * Rotation units are degrees in both OpenTrackIO and Unreal.
+	 */
+
+	/** Factor to convert from OpenTrackIO's meters to Unreal's centimeters */
+	inline constexpr float MetersToCentimeters = 100.f;
+}
+
 /**
  * Represents frame rate type from OpenTrack I/O. Value is represented in hertz
  */
@@ -1065,12 +1090,13 @@ struct FLiveLinkOpenTrackIO_Rotator
 	
 	bool operator==(const FLiveLinkOpenTrackIO_Rotator&) const = default;
 	bool operator!=(const FLiveLinkOpenTrackIO_Rotator&) const = default;
-
-	FRotator GetRotator() const
+	
+	/** From OpenTrackIO to Unreal Translation: Invert sign of Yaw */
+	FRotator ToUnrealRotator() const
 	{
-		return FRotator(Tilt,Pan,Roll);
+		return FRotator(Tilt, -Pan, Roll);
 	}
-	
+
 	UPROPERTY()
 	float Pan = 0;
 
@@ -1102,9 +1128,16 @@ struct FLiveLinkOpenTrackIO_XYZ
 	bool operator==(const FLiveLinkOpenTrackIO_XYZ&) const = default;
 	bool operator!=(const FLiveLinkOpenTrackIO_XYZ&) const = default;
 
-	FVector GetVector() const
+	/** From OpenTrackIO to Unreal Translation: Swap X <-> Y and convert meters to cm */
+	FVector ToUnrealTranslation() const
+	{
+		return FVector(Y, X, Z) * LiveLinkOpenTrackIO::MetersToCentimeters;
+	}
+
+	/** From OpenTrackIO to Unreal Scale: Swap X <-> Y (unitless) */
+	FVector ToUnrealScale() const
 	{
-		return FVector(X,Y,Z);
+		return FVector(Y, X, Z);
 	}
 
 	UPROPERTY()
@@ -1132,13 +1165,14 @@ struct FLiveLinkOpenTrackIOTransform
 	bool operator==(const FLiveLinkOpenTrackIOTransform&) const = default;
 	bool operator!=(const FLiveLinkOpenTrackIOTransform&) const = default;
 
-	FTransform GetTransform() const
+	/** From OpenTrackIO to Unreal FTransform */
+	FTransform ToUnrealTransform() const
 	{
-		FTransform Transform(
-			Rotation.GetRotator(),
-			Translation.GetVector(),
-			Scale.GetVector());
-		return Transform;
+		return FTransform(
+			Rotation.ToUnrealRotator(),
+			Translation.ToUnrealTranslation(),
+			Scale.ToUnrealScale()
+		);
 	}
 	
 	/** Translation part of a OpenTrackIO xform */
