OpenTrackIO: Added Source ID to connection settings dialog.

This value will be automatically used to populate the 4th octet of the multicast IP address, per the spec.
The user is not allowed to change the multicast ip address directly, which avoids discrepancy from the spec and also the connection string source id not maching anymore the 4th octet of the multicast ip address.

For now, not including the port number in the connection string. If the user requires a port different than the spec's default 55555, it can be changed in the Live Link Source settings after creating it.

diff --git a/Engine/Plugins/VirtualProduction/LiveLinkOpenTrackIO/Source/LiveLinkOpenTrackIO/Private/LiveLinkOpenTrackIOConnectionSettings.h b/Engine/Plugins/VirtualProduction/LiveLinkOpenTrackIO/Source/LiveLinkOpenTrackIO/Private/LiveLinkOpenTrackIOConnectionSettings.h
--- a/Engine/Plugins/VirtualProduction/LiveLinkOpenTrackIO/Source/LiveLinkOpenTrackIO/Private/LiveLinkOpenTrackIOConnectionSettings.h
+++ b/Engine/Plugins/VirtualProduction/LiveLinkOpenTrackIO/Source/LiveLinkOpenTrackIO/Private/LiveLinkOpenTrackIOConnectionSettings.h
@@ -9,4 +9,8 @@ USTRUCT()
 struct LIVELINKOPENTRACKIO_API FLiveLinkOpenTrackIOConnectionSettings
 {
 	GENERATED_BODY()
+
+	/** The Source Number is a user-configurable 8-bit value that determines the multicast IP address for a specific Source from a Producer. */
+	UPROPERTY(EditAnywhere, meta = (ClampMin = "1", ClampMax = "200", UIMin = "1", UIMax = "200"), Category = "Settings")
+	uint8 SourceNumber = 1;
 };
diff --git a/Engine/Plugins/VirtualProduction/LiveLinkOpenTrackIO/Source/LiveLinkOpenTrackIO/Private/LiveLinkOpenTrackIOSource.cpp b/Engine/Plugins/VirtualProduction/LiveLinkOpenTrackIO/Source/LiveLinkOpenTrackIO/Private/LiveLinkOpenTrackIOSource.cpp
--- a/Engine/Plugins/VirtualProduction/LiveLinkOpenTrackIO/Source/LiveLinkOpenTrackIO/Private/LiveLinkOpenTrackIOSource.cpp
+++ b/Engine/Plugins/VirtualProduction/LiveLinkOpenTrackIO/Source/LiveLinkOpenTrackIO/Private/LiveLinkOpenTrackIOSource.cpp
@@ -100,13 +100,33 @@ void JoinedToGroup(const FIPv4Endpoint& UnicastEndpoint, const FIPv4Endpoint& Mu
 #endif
 }
 
+/**
+ * Makes the Multicast Endpoint to use per the opentrackio spec.
+ * Source Number will be clamped to 1..200 per opentrackio.org.
+ * 
+ * @param InSourceNumber  OpenTrackIO Source Number to encode into the fourth octet.
+ * @param InPortNumber    UDP multicast port number.
+ *
+ * @return FIPv4Endpoint  The multicast address (235.135.1.<SourceNumber>:<PortNumber>).
+ */
+FIPv4Endpoint MakeMulticastAddressFromSourceNumber(uint8 InSourceNumber, uint16 InPortNumber = 55555)
+{
+	const uint8 SourceNumber = FMath::Clamp(InSourceNumber, 1, 200);
+	return FIPv4Endpoint(FIPv4Address(235, 135, 1, SourceNumber), InPortNumber);
+}
+
 }
 
 FLiveLinkOpenTrackIOSource::FLiveLinkOpenTrackIOSource(FLiveLinkOpenTrackIOConnectionSettings InConnectionSettings)
 	: ConnectionSettings(MoveTemp(InConnectionSettings))
 {
 	SourceType = LOCTEXT("SourceType_OpenTrackIO", "OpenTrack I/O");
-	SourceMachineName = LOCTEXT("LocalMachine", "Local");
+
+	SourceMachineName = FText::Format(
+		LOCTEXT("SourceNumber", "Source # {0}"),
+		ConnectionSettings.SourceNumber
+	);
+
 	SourceStatus = LOCTEXT("Startup", "Waiting for startup...");
 	DeferredStartDelegateHandle = FCoreDelegates::OnEndFrame.AddRaw(this, &FLiveLinkOpenTrackIOSource::Start);
 }
@@ -200,12 +220,9 @@ void FLiveLinkOpenTrackIOSource::ParseEndpointFromSourceSettings()
 	using namespace UE::LiveLinkOpenTrackIO::Private; 
 	if (TOptional<FIPv4Endpoint> InUniEndpoint = ParseEndpoint(SavedSourceSettings->UnicastEndpoint))
 	{
-		if (TOptional<FIPv4Endpoint> InMultiEndpoint = ParseEndpoint(SavedSourceSettings->MulticastEndpoint))
-		{
-			UnicastEndpoint = *InUniEndpoint;
-			MulticastEndpoint = *InMultiEndpoint;
-			ConnectionState = ELiveLinkOpenTrackIOState::EndpointsReady;
-		}
+		UnicastEndpoint = *InUniEndpoint;
+		MulticastEndpoint = MakeMulticastAddressFromSourceNumber(ConnectionSettings.SourceNumber, SavedSourceSettings->MulticastPort);
+		ConnectionState = ELiveLinkOpenTrackIOState::EndpointsReady;
 	}
 }
 
@@ -251,7 +268,7 @@ void FLiveLinkOpenTrackIOSource::OnSettingsChanged(ULiveLinkSourceSettings* Sett
 			const FName PropertyName = Property->GetFName();
 			const FName MemberPropertyName = MemberProperty->GetFName();
 			if (PropertyName == GET_MEMBER_NAME_CHECKED(ULiveLinkOpenTrackIOSourceSettings, UnicastEndpoint)
-				|| PropertyName == GET_MEMBER_NAME_CHECKED(ULiveLinkOpenTrackIOSourceSettings, MulticastEndpoint))
+				|| PropertyName == GET_MEMBER_NAME_CHECKED(ULiveLinkOpenTrackIOSourceSettings, MulticastPort))
 			{
 				/** Start the reset process. */
 				ResetSockets();
@@ -397,7 +414,7 @@ bool FLiveLinkOpenTrackIOSource::OpenSockets()
 		UE_LOG(LogLiveLinkOpenTrackIO, Display, TEXT("OpenTrackIO unicast socket socket bound to '%s'."), *IpPortStr);
 
 		AsyncTask(ENamedThreads::GameThread, [IpPortStr, this]() {
-			SourceMachineName = FText::Format(FTextFormat(LOCTEXT("LocalMachineWithIP", "Local ({0})")), {FText::FromString(IpPortStr)});
+
 			if (IsValid(SavedSourceSettings))
 			{
 				SavedSourceSettings->AssignedUnicastEndpoint = IpPortStr;
diff --git a/Engine/Plugins/VirtualProduction/LiveLinkOpenTrackIO/Source/LiveLinkOpenTrackIO/Public/LiveLinkOpenTrackIOSourceSettings.h b/Engine/Plugins/VirtualProduction/LiveLinkOpenTrackIO/Source/LiveLinkOpenTrackIO/Public/LiveLinkOpenTrackIOSourceSettings.h
--- a/Engine/Plugins/VirtualProduction/LiveLinkOpenTrackIO/Source/LiveLinkOpenTrackIO/Public/LiveLinkOpenTrackIOSourceSettings.h
+++ b/Engine/Plugins/VirtualProduction/LiveLinkOpenTrackIO/Source/LiveLinkOpenTrackIO/Public/LiveLinkOpenTrackIOSourceSettings.h
@@ -12,15 +12,12 @@ class LIVELINKOPENTRACKIO_API ULiveLinkOpenTrackIOSourceSettings : public ULiveL
 
 public:
 	/**
-	 * The IP endpoint to send multicast packets to.
+	 * The Multicast port to receive packets on.
 	 *
-	 * The format is IP_ADDRESS:PORT_NUMBER.
-	 * The multicast IP address must be in the range 224.0.0.0 to 239.255.255.255.
-	 *
-	 * The default value is 235.135.1.14:55555 per OpenTrackIO specifications.
+	 * The default value 55555 per OpenTrackIO specifications.
 	 */
 	UPROPERTY(EditAnywhere, Category=Transport)
-	FString MulticastEndpoint = TEXT("235.135.1.14:55555");
+	uint16 MulticastPort = 55555;
 
 	/**
 	 * The IP endpoint to listen to and send packets from.  
