OpenTrackIO: Add default subject name while static camera data arrives.

It is possible that an OpenTrackIO source never sends camera data (e.g. make/model/label) but we still need a way to represent it in Live Link.

For this, we initially create a Subject with a default name based on the Source Number, and if the camera data later arrives, we replace it.

diff --git a/Engine/Plugins/VirtualProduction/LiveLinkOpenTrackIO/Source/LiveLinkOpenTrackIO/Private/LiveLinkOpenTrackIOSource.cpp b/Engine/Plugins/VirtualProduction/LiveLinkOpenTrackIO/Source/LiveLinkOpenTrackIO/Private/LiveLinkOpenTrackIOSource.cpp
--- a/Engine/Plugins/VirtualProduction/LiveLinkOpenTrackIO/Source/LiveLinkOpenTrackIO/Private/LiveLinkOpenTrackIOSource.cpp
+++ b/Engine/Plugins/VirtualProduction/LiveLinkOpenTrackIO/Source/LiveLinkOpenTrackIO/Private/LiveLinkOpenTrackIOSource.cpp
@@ -115,6 +115,19 @@ FIPv4Endpoint MakeMulticastAddressFromSourceNumber(uint8 InSourceNumber, uint16
 	return FIPv4Endpoint(FIPv4Address(235, 135, 1, SourceNumber), InPortNumber);
 }
 
+/**
+ * Makes an interim camera label based on the OpenTrackIO Source Number.
+ * Useful when the producer hasn't sent static camera data yet.
+ *
+ * @param InSourceNumber OpenTrackIO Source Number to insert into the subject name.
+ * 
+ * @return FString In the form "OpenTrackIO_<SourceNumber>".
+ */
+FString MakeInterimCameraLabelFromSourceNumber(uint8 InSourceNumber)
+{
+	return FString::Printf(TEXT("OpenTrackIO_%03u"), InSourceNumber);
+}
+
 }
 
 FLiveLinkOpenTrackIOSource::FLiveLinkOpenTrackIOSource(FLiveLinkOpenTrackIOConnectionSettings InConnectionSettings)
@@ -338,9 +351,42 @@ void FLiveLinkOpenTrackIOSource::PushDataToLiveLink_AnyThread(const FLiveLinkOpe
 	/** Check the static data camera fields on the inbound struct and see if we have new data to store / update. */
 	if (TOptional<FLiveLinkStaticDataStruct> StaticData = Cache->TryAddCamera(InData))
 	{
-		FName SubjectName = Cache->GetSubjectNameFromData(InData);
+		const FName SubjectName = Cache->GetSubjectNameFromData(InData);
 
 		Client->PushSubjectStaticData_AnyThread({ SourceGuid, SubjectName }, ULiveLinkLensRole::StaticClass(), MoveTemp(*StaticData));
+
+		// Remove the interim live link subject, if any.
+		if (!DefaultSubjectName.IsNone())
+		{
+			Client->RemoveSubject_AnyThread({ SourceGuid, DefaultSubjectName });
+			DefaultSubjectName = NAME_None;
+		}
+	}
+	else
+	{
+		// If no static camera data has arrived, we add an interim one named after the OpenTrackIO Source Number to support producers
+		// that do not send this information. If valid static camera data arrives later, the interim live link subject will be removed
+		// and replaced with the new one.
+		if (!Cache->StaticCameras.Num())
+		{
+			using namespace UE::LiveLinkOpenTrackIO::Private;
+
+			FLiveLinkOpenTrackIOData InterimCameraData = InData;
+
+			InterimCameraData.Static.Camera.Label = MakeInterimCameraLabelFromSourceNumber(ConnectionSettings.SourceNumber);
+			InterimCameraData.Static.Camera.Make = "N/A";
+			InterimCameraData.Static.Camera.Model = "N/A";
+			InterimCameraData.Static.Camera.SerialNumber = "N/A";
+
+			StaticData = Cache->TryAddCamera(InterimCameraData);
+
+			if (ensure(StaticData.IsSet())) // Should not fail since we populated the static camera data.
+			{
+				DefaultSubjectName = Cache->GetSubjectNameFromData(InterimCameraData);
+
+				Client->PushSubjectStaticData_AnyThread({ SourceGuid, DefaultSubjectName }, ULiveLinkLensRole::StaticClass(), MoveTemp(*StaticData));
+			}
+		}
 	}
 
 	TArray<FSubjectAndFrameDataTuple> FrameData = Cache->TryReadFrameData(InData);
diff --git a/Engine/Plugins/VirtualProduction/LiveLinkOpenTrackIO/Source/LiveLinkOpenTrackIO/Private/LiveLinkOpenTrackIOSource.h b/Engine/Plugins/VirtualProduction/LiveLinkOpenTrackIO/Source/LiveLinkOpenTrackIO/Private/LiveLinkOpenTrackIOSource.h
--- a/Engine/Plugins/VirtualProduction/LiveLinkOpenTrackIO/Source/LiveLinkOpenTrackIO/Private/LiveLinkOpenTrackIOSource.h
+++ b/Engine/Plugins/VirtualProduction/LiveLinkOpenTrackIO/Source/LiveLinkOpenTrackIO/Private/LiveLinkOpenTrackIOSource.h
@@ -193,4 +193,6 @@ private:
 	/** Counter to track the number of inbound OpenTrackIO packets received. */
 	std::atomic<uint64> OpenTrackIOPacketCount = 0;
 	
+	/** Default subject name, used while no static camera data has arrived */
+	FName DefaultSubjectName;
 };
