- Update the LiveLinkCameraController to support dynamic filmback resolution from a frame data.
- Update the LiveLinkOpenTrackIOTypes to support Blueprint.
- Add fields for custom data that we will push into live link frame data.
- Fix unit test data to conform with mod 255 checksum calculation.

diff --git a/Engine/Plugins/VirtualProduction/LiveLinkCamera/Source/LiveLinkCamera/Private/LiveLinkCameraController.cpp b/Engine/Plugins/VirtualProduction/LiveLinkCamera/Source/LiveLinkCamera/Private/LiveLinkCameraController.cpp
--- a/Engine/Plugins/VirtualProduction/LiveLinkCamera/Source/LiveLinkCamera/Private/LiveLinkCameraController.cpp
+++ b/Engine/Plugins/VirtualProduction/LiveLinkCamera/Source/LiveLinkCamera/Private/LiveLinkCameraController.cpp
@@ -80,14 +80,21 @@ void ULiveLinkCameraController::Tick(float DeltaTime, const FLiveLinkSubjectFram
 				ULensFile* SelectedLensFile = LensFilePicker.GetLensFile();
 				LensFileEvalData.LensFile = SelectedLensFile;
 
-				if (StaticData->FilmBackWidth > 0.0f && UpdateFlags.bApplyFilmBack) 
-				{ 
-					CineCameraComponent->Filmback.SensorWidth = StaticData->FilmBackWidth; 
-				}
-
-				if (StaticData->FilmBackHeight > 0.0f && UpdateFlags.bApplyFilmBack) 
-				{ 
-					CineCameraComponent->Filmback.SensorHeight = StaticData->FilmBackHeight; 
+				auto ConditionallyApplySensorChanges = [StaticData](float& Target, float StaticValue, float FrameValue)
+				{
+					if (StaticData->bIsDynamicFilmbackSupported && FrameValue > 0.0f)
+					{
+						Target = FrameValue;
+					}
+					else if (StaticValue > 0.0f)
+					{
+						Target = StaticValue;
+					}
+				};
+				if (UpdateFlags.bApplyFilmBack)
+				{
+					ConditionallyApplySensorChanges(CineCameraComponent->Filmback.SensorWidth, StaticData->FilmBackWidth, FrameData->FilmBackWidth);
+					ConditionallyApplySensorChanges(CineCameraComponent->Filmback.SensorHeight, StaticData->FilmBackHeight, FrameData->FilmBackHeight);
 				}
 
 				//Verify different Lens Tables based on streamed FIZ at some intervals
diff --git a/Engine/Plugins/VirtualProduction/LiveLinkCamera/Source/LiveLinkCamera/Public/LiveLinkCameraController.h b/Engine/Plugins/VirtualProduction/LiveLinkCamera/Source/LiveLinkCamera/Public/LiveLinkCameraController.h
--- a/Engine/Plugins/VirtualProduction/LiveLinkCamera/Source/LiveLinkCamera/Public/LiveLinkCameraController.h
+++ b/Engine/Plugins/VirtualProduction/LiveLinkCamera/Source/LiveLinkCamera/Public/LiveLinkCameraController.h
@@ -38,7 +38,7 @@ struct FLiveLinkCameraControllerUpdateFlags
 	UPROPERTY(EditAnywhere, Category = "Updates")
 	bool bApplyProjectionMode= true;
 
-	/** Whether to apply Filmback if it's available in LiveLink StaticData */
+	/** Whether to apply Filmback if it's available in LiveLink StaticData or sent in the LiveLink FrameData*/
 	UPROPERTY(EditAnywhere, Category = "Updates")
 	bool bApplyFilmBack = true;
 
diff --git a/Engine/Plugins/VirtualProduction/LiveLinkOpenTrackIO/Source/LiveLinkOpenTrackIO/Private/LiveLinkOpenTrackIOTranscoder.cpp b/Engine/Plugins/VirtualProduction/LiveLinkOpenTrackIO/Source/LiveLinkOpenTrackIO/Private/LiveLinkOpenTrackIOTranscoder.cpp
--- a/Engine/Plugins/VirtualProduction/LiveLinkOpenTrackIO/Source/LiveLinkOpenTrackIO/Private/LiveLinkOpenTrackIOTranscoder.cpp
+++ b/Engine/Plugins/VirtualProduction/LiveLinkOpenTrackIO/Source/LiveLinkOpenTrackIO/Private/LiveLinkOpenTrackIOTranscoder.cpp
@@ -22,6 +22,7 @@ bool AddStaticData(FName SourceId, const T& Item, TMap<FName, T>& Container)
 		FoundItem = Item;
 		return true;
 	}
+	FoundItem.UpdateStaticData(Item);
 
 	return false;
 }
@@ -29,7 +30,8 @@ bool AddStaticData(FName SourceId, const T& Item, TMap<FName, T>& Container)
 
 FName FLiveLinkOpenTrackIOCache::GetSubjectNameFromData(const FLiveLinkOpenTrackIOData& InData)
 {
-	return UE::OpenTrackIO::ConvertTypeToFName(InData.Static.Camera); 
+	const FString SourceNumberAsString = FString::FromInt(InData.SourceNumber);
+	return UE::OpenTrackIO::ConvertTypeToFName(InData.Static.Camera, SourceNumberAsString); 
 }
 
 TOptional<FLiveLinkStaticDataStruct> FLiveLinkOpenTrackIOCache::TryAddCamera(const FLiveLinkOpenTrackIOData& Data)
@@ -49,6 +51,8 @@ TOptional<FLiveLinkStaticDataStruct> FLiveLinkOpenTrackIOCache::TryAddCamera(con
 	LensData.bIsFocalLengthSupported = true;
 	LensData.bIsApertureSupported = true;
 	LensData.bIsFocusDistanceSupported = true;
+	LensData.bIsDynamicFilmbackSupported = true;
+	
 	LensData.FilmBackHeight = Data.Static.Camera.ActiveSensorPhysicalDimensions.Height;
 	LensData.FilmBackWidth = Data.Static.Camera.ActiveSensorPhysicalDimensions.Width;
 
@@ -100,16 +104,16 @@ TArray<FSubjectAndFrameDataTuple> FLiveLinkOpenTrackIOCache::TryReadFrameData(co
 		SampleTC = Data.Timing.Timecode.GetQualifiedFrameTime();
 	}
 
-	TArray<FName> Keys;
-	StaticCameras.GenerateKeyArray(Keys);
-
 	double WorldTime = FPlatformTime::Seconds();
-	for (const FName SourceId : Keys)
+
+	const FName SourceId = FName(Data.SourceId);
+	
+	// We require cameras to exist but it is optional to have a lens.
+	//
+	const FLiveLinkOpenTrackIOStaticCamera* Camera = StaticCameras.Find(SourceId);
+	if (Camera)
 	{
-		// We require cameras to exist but it is optional to have a lens.
-		// 
-		FLiveLinkOpenTrackIOStaticCamera& Camera = StaticCameras.FindChecked(SourceId);
-		FLiveLinkOpenTrackIOStaticLens* Lens = StaticLenses.Find(SourceId);
+		const FLiveLinkOpenTrackIOStaticLens* Lens = StaticLenses.Find(SourceId);
 		
 		FLiveLinkFrameDataStruct FrameDataStruct = FLiveLinkFrameDataStruct(FLiveLinkLensFrameData::StaticStruct());
 		FLiveLinkLensFrameData& CameraAndLensData = *FrameDataStruct.Cast<FLiveLinkLensFrameData>();
@@ -122,12 +126,18 @@ TArray<FSubjectAndFrameDataTuple> FLiveLinkOpenTrackIOCache::TryReadFrameData(co
 		CameraAndLensData.Aperture = Data.Lens.FStop;
 		CameraAndLensData.FocalLength = Data.Lens.PinholeFocalLength;
 
+		CameraAndLensData.FilmBackHeight = Camera->ActiveSensorPhysicalDimensions.Height;
+		CameraAndLensData.FilmBackWidth = Camera->ActiveSensorPhysicalDimensions.Width;
+		
 		if (Lens)
 		{
-			CameraAndLensData.PrincipalPoint = Data.Lens.DistortionOffset.AsVector(); 
+			CameraAndLensData.PrincipalPoint = Data.Lens.DistortionOffset.AsVector();
 		}
+		FName SubjectName = GetSubjectNameFromData(Data);
 
-		FName SubjectName = UE::OpenTrackIO::ConvertTypeToFName(Camera);
+		// Note that although we are only adding frame data for a single subject. This function returns an array to
+		// support sending frame data for other parts of the OpenTrackIO spec.
+		// 
 		Result.Add({SubjectName, MoveTemp(FrameDataStruct)});
 	}
 
diff --git a/Engine/Plugins/VirtualProduction/LiveLinkOpenTrackIO/Source/LiveLinkOpenTrackIO/Private/LiveLinkOpenTrackIOTypes.h b/Engine/Plugins/VirtualProduction/LiveLinkOpenTrackIO/Source/LiveLinkOpenTrackIO/Private/LiveLinkOpenTrackIOTypes.h
--- a/Engine/Plugins/VirtualProduction/LiveLinkOpenTrackIO/Source/LiveLinkOpenTrackIO/Private/LiveLinkOpenTrackIOTypes.h
+++ b/Engine/Plugins/VirtualProduction/LiveLinkOpenTrackIO/Source/LiveLinkOpenTrackIO/Private/LiveLinkOpenTrackIOTypes.h
@@ -8,6 +8,8 @@
 #include "Math/Matrix.h"
 #include "Math/Transform.h"
 
+#include "Kismet/BlueprintFunctionLibrary.h"
+
 #include "LiveLinkOpenTrackIOTypes.generated.h"
 
 
@@ -36,9 +38,10 @@ namespace LiveLinkOpenTrackIO
 }
 
 /**
- * Represents frame rate type from OpenTrack I/O. Value is represented in hertz
+ * Represents frame rate type from OpenTrack I/O. Value is represented in hertz using an fractional representation
+ * Numerator / Denominator. The default value is 24 / 1 or 24 cycles per second.
  */
-USTRUCT()
+USTRUCT(BlueprintType)
 struct FLiveLinkOpenTrackIOFrameRate
 {
 	GENERATED_BODY()
@@ -56,14 +59,16 @@ struct FLiveLinkOpenTrackIOFrameRate
 	FFrameRate GetFrameRate() const
 	{
 		check(Num > 0);
-		return FFrameRate(static_cast<uint32>(Num), Denom);
+		return FFrameRate(static_cast<uint32>(Num), static_cast<uint32>(Denom));
 	}
 
-	UPROPERTY()
+	/** Numerator of the frame rate value. */
+	UPROPERTY(EditAnywhere, BlueprintReadWrite, Category = "OpenTrackIO")
 	int32 Num = 24;
 
-	UPROPERTY()
-	uint32 Denom = 1;
+	/** Denominator of the frame rate value. */
+	UPROPERTY(EditAnywhere, BlueprintReadWrite, Category = "OpenTrackIO")
+	int32 Denom = 1;
 };
 
 /** Concepts to test for the existance of member variables on our structs. */
@@ -86,7 +91,7 @@ namespace UE::OpenTrackIO
 {
 
 template<CHasModelAndMake T>
-FName ConvertTypeToFName(const T& Type)
+FName ConvertTypeToFName(const T& Type, const FString& Postfix)
 {
 	if constexpr(CHasLabelMember<T>)
 	{
@@ -95,9 +100,11 @@ FName ConvertTypeToFName(const T& Type)
 	else
 	{
 		FNameBuilder Builder;
-		Builder.Append(*Type.Make);
-		Builder.Append(TEXT("_"));
-		Builder.Append(*Type.Model);
+		Builder.Appendf(TEXT("%s_%s"), *Type.Make, *Type.Model);
+		if (!Postfix.IsEmpty())
+		{
+			Builder.Appendf(TEXT(":%s"), *Postfix);
+		}
 		return *Builder;
 	}
 }
@@ -106,7 +113,7 @@ FName ConvertTypeToFName(const T& Type)
 /**
  * Height and width of the active area of the camera sensor in millimeters
  */
-USTRUCT()
+USTRUCT(BlueprintType)
 struct FLiveLinkOpenTrackIOSensorDimensions
 {
 	GENERATED_BODY()
@@ -121,17 +128,17 @@ struct FLiveLinkOpenTrackIOSensorDimensions
 	bool operator==(const FLiveLinkOpenTrackIOSensorDimensions&) const = default;
 	bool operator!=(const FLiveLinkOpenTrackIOSensorDimensions&) const = default;
 	
-	UPROPERTY()
+	UPROPERTY(EditAnywhere, BlueprintReadWrite, Category = "OpenTrackIO")
 	float Height = 43.0f;
 
-	UPROPERTY()
+	UPROPERTY(EditAnywhere, BlueprintReadWrite, Category = "OpenTrackIO")
 	float Width = 24.0f;
 };
 
 /**
  *  Photosite resolution of the active area of the camera sensor in pixels
  */
-USTRUCT()
+USTRUCT(BlueprintType)
 struct FLiveLinkOpenTrackIOSensorResolution
 {
 	GENERATED_BODY()
@@ -146,10 +153,10 @@ struct FLiveLinkOpenTrackIOSensorResolution
 	bool operator==(const FLiveLinkOpenTrackIOSensorResolution&) const = default;
 	bool operator!=(const FLiveLinkOpenTrackIOSensorResolution&) const = default;
 	
-	UPROPERTY()
+	UPROPERTY(EditAnywhere, BlueprintReadWrite, Category = "OpenTrackIO")
 	int32 Height = 1920;
 
-	UPROPERTY()
+	UPROPERTY(EditAnywhere, BlueprintReadWrite, Category = "OpenTrackIO")
 	int32 Width = 1080;
 };
 
@@ -158,7 +165,7 @@ struct FLiveLinkOpenTrackIOSensorResolution
  * It can be used to de-squeeze images but is not however an exact number over the entire captured
  * area due to a lens’ intrinsic analog nature.
  */
-USTRUCT()
+USTRUCT(BlueprintType)
 struct FLiveLinkOpenTrackIOAnamorphicSqueeze
 {
 	GENERATED_BODY()
@@ -173,14 +180,14 @@ struct FLiveLinkOpenTrackIOAnamorphicSqueeze
 	bool operator==(const FLiveLinkOpenTrackIOAnamorphicSqueeze&) const = default;
 	bool operator!=(const FLiveLinkOpenTrackIOAnamorphicSqueeze&) const = default;
 	
-	UPROPERTY()
+	UPROPERTY(EditAnywhere, BlueprintReadWrite, Category = "OpenTrackIO")
 	int32 Num = 1;
 
-	UPROPERTY()
-	uint32 Denom = 1;		
+	UPROPERTY(EditAnywhere, BlueprintReadWrite, Category = "OpenTrackIO")
+	int32 Denom = 1;		
 };
 
-USTRUCT()
+USTRUCT(BlueprintType)
 struct FLiveLinkOpenTrackIOStaticCamera
 {
 	GENERATED_BODY()
@@ -191,61 +198,79 @@ struct FLiveLinkOpenTrackIOStaticCamera
 		static FLiveLinkOpenTrackIOStaticCamera DefaultValue;
 		return DefaultValue == *this;
 	}
-	
-	bool operator==(const FLiveLinkOpenTrackIOStaticCamera&) const = default;
-	bool operator!=(const FLiveLinkOpenTrackIOStaticCamera&) const = default;
+
+	bool operator==(const FLiveLinkOpenTrackIOStaticCamera& Other) const
+	{
+		return Other.Make == Make
+			&& Other.Model == Model
+			&& Other.SerialNumber == SerialNumber
+			&& Other.Label == Label
+			&& Other.FdlLink == FdlLink
+			&& Other.FirmwareVersion == FirmwareVersion; 
+	}
+	bool operator!=(const FLiveLinkOpenTrackIOStaticCamera& Other) const
+	{
+		return !(Other == *this);
+	}
 	
 	/** If the provided static camera specifies the correct properties to be considered valid. */
 	bool IsValid() const
 	{
 		return !Make.IsEmpty() && !Model.IsEmpty() && !SerialNumber.IsEmpty() && !Label.IsEmpty();
 	}
+
+	void UpdateStaticData(const FLiveLinkOpenTrackIOStaticCamera& Other)
+	{
+		ActiveSensorPhysicalDimensions = Other.ActiveSensorPhysicalDimensions;
+		IsoSpeed = Other.IsoSpeed;
+		ShutterAngle = Other.ShutterAngle;
+	}
 	
 	/** Object representing the capture rate.  */
-	UPROPERTY()
+	UPROPERTY(EditAnywhere, BlueprintReadWrite, Category = "OpenTrackIO")
 	FLiveLinkOpenTrackIOFrameRate CaptureFrameRate;
 
 	/** Object representing the sensor dimensions. */
-	UPROPERTY()
+	UPROPERTY(EditAnywhere, BlueprintReadWrite, Category = "OpenTrackIO")
 	FLiveLinkOpenTrackIOSensorDimensions ActiveSensorPhysicalDimensions;
 
 	/** Object representing the sensor resolution.  */
-	UPROPERTY()
+	UPROPERTY(EditAnywhere, BlueprintReadWrite, Category = "OpenTrackIO")
 	FLiveLinkOpenTrackIOSensorResolution ActiveSensorResolution;
 
 	/** Non-blank string naming camera manufacturer. */
-	UPROPERTY()
+	UPROPERTY(EditAnywhere, BlueprintReadWrite, Category = "OpenTrackIO")
 	FString Make;
 
 	/** Non-blank string naming camera model. */
-	UPROPERTY()
+	UPROPERTY(EditAnywhere, BlueprintReadWrite, Category = "OpenTrackIO")
 	FString Model;
 
 	/** Non-blank string for camera serial number. */
-	UPROPERTY()
+	UPROPERTY(EditAnywhere, BlueprintReadWrite, Category = "OpenTrackIO")
 	FString SerialNumber;
 
 	/** Non-blank string identifying camera firmware version */
-	UPROPERTY()
+	UPROPERTY(EditAnywhere, BlueprintReadWrite, Category = "OpenTrackIO")
 	FString FirmwareVersion;
 
 	/** Non-blank string containing user-determined camera identifier */
-	UPROPERTY()
+	UPROPERTY(EditAnywhere, BlueprintReadWrite, Category = "OpenTrackIO")
 	FString Label;
 
 	/** Anamorphic squeeze ratio */
-	UPROPERTY()
+	UPROPERTY(EditAnywhere, BlueprintReadWrite, Category = "OpenTrackIO")
 	FLiveLinkOpenTrackIOAnamorphicSqueeze AnamorphicSqueeze;
 
 	/** Arithmetic ISO scale as defined in ISO 12232 */
-	UPROPERTY()
-	uint32 IsoSpeed = 400;
+	UPROPERTY(EditAnywhere, BlueprintReadWrite, Category = "OpenTrackIO")
+	int32 IsoSpeed = 400;
 
 	/**
 	 * URN identifying the ASC Framing Decision List used by the camera.
 	 * Pattern:  ^urn:uuid:[0-9a-f]{8}-[0-9a-f]{4}-[0-9a-f]{4}-[0-9a-f]{4}-[0-9a-f]{12}$
 	 */
-	UPROPERTY()
+	UPROPERTY(EditAnywhere, BlueprintReadWrite, Category = "OpenTrackIO")
 	FString FdlLink;
 
 	/**
@@ -256,7 +281,7 @@ struct FLiveLinkOpenTrackIOStaticCamera
 	float ShutterAngle = 180.0;
 };
 
-USTRUCT()
+USTRUCT(BlueprintType)
 struct FLiveLinkOpenTrackIOStaticLens
 {
 	GENERATED_BODY()
@@ -277,12 +302,17 @@ struct FLiveLinkOpenTrackIOStaticLens
 		return !Make.IsEmpty() && !Model.IsEmpty() && !SerialNumber.IsEmpty();
 	}
 
+	void UpdateStaticData(const FLiveLinkOpenTrackIOStaticLens&)
+	{
+		// TODO No-op for now. 
+	}
+	
 	/**
 	 * Static maximum overscan factor on lens distortion. This is an alternative to
 	 * providing dynamic overscan values each frame. Note it should be the maximum
 	 * of both projection-matrix-based and field-of-view-based rendering as per the OpenLensIO documentation.
 	 */
-	UPROPERTY()
+	UPROPERTY(EditAnywhere, BlueprintReadWrite, Category = "OpenTrackIO")
 	float DistortionOverscanMax = 1.0;
 
 	/**
@@ -291,35 +321,35 @@ struct FLiveLinkOpenTrackIOStaticLens
 	 * projection-matrix-based and field-of-view-based rendering as per the OpenLensIO
 	 * documentation.
 	 */
-	UPROPERTY()
+	UPROPERTY(EditAnywhere, BlueprintReadWrite, Category = "OpenTrackIO")
 	float UndistortionOverscanMax = 1.0;
 
 	/** Non-blank string naming lens manufacturer */
-	UPROPERTY()
+	UPROPERTY(EditAnywhere, BlueprintReadWrite, Category = "OpenTrackIO")
 	FString Make;
 
 	/** Non-blank string identifying lens model */
-	UPROPERTY()
+	UPROPERTY(EditAnywhere, BlueprintReadWrite, Category = "OpenTrackIO")
 	FString Model;
 
 	/** Non-blank string uniquely identifying the lens */
-	UPROPERTY()
+	UPROPERTY(EditAnywhere, BlueprintReadWrite, Category = "OpenTrackIO")
 	FString SerialNumber;
 
 	/** Non-blank string identifying lens firmware version */
-	UPROPERTY()
+	UPROPERTY(EditAnywhere, BlueprintReadWrite, Category = "OpenTrackIO")
 	FString FirmwareVersion;
 
 	/** Nominal focal length of the lens. The number printed on the side of a prime lens, e.g. 50 mm, and undefined in the case of a zoom lens. */
-	UPROPERTY()
+	UPROPERTY(EditAnywhere, BlueprintReadWrite, Category = "OpenTrackIO")
 	float NominalFocalLength = 50;
 
 	/** List of free strings that describe the history of calibrations of the lens. */
-	UPROPERTY()
+	UPROPERTY(EditAnywhere, BlueprintReadWrite, Category = "OpenTrackIO")
 	TArray<FString> CalibrationHistory;
 };
 
-USTRUCT()
+USTRUCT(BlueprintType)
 struct FLiveLinkOpenTrackIOStaticTracker
 {
 	GENERATED_BODY()
@@ -341,23 +371,23 @@ struct FLiveLinkOpenTrackIOStaticTracker
 	}
 
 	/** Non-blank string naming tracking device manufacturer */
-	UPROPERTY()
+	UPROPERTY(EditAnywhere, BlueprintReadWrite, Category = "OpenTrackIO")
 	FString Make;
 
 	/**  Non-blank string identifying tracking device model */
-	UPROPERTY()
+	UPROPERTY(EditAnywhere, BlueprintReadWrite, Category = "OpenTrackIO")
 	FString Model;
 
 	/** Non-blank string uniquely identifying the tracking device */
-	UPROPERTY()
+	UPROPERTY(EditAnywhere, BlueprintReadWrite, Category = "OpenTrackIO")
 	FString SerialNumber;
 
 	/**  Non-blank string identifying tracking device firmware version */
-	UPROPERTY()
+	UPROPERTY(EditAnywhere, BlueprintReadWrite, Category = "OpenTrackIO")
 	FString FirmwareVersion;
 };
 
-USTRUCT()
+USTRUCT(BlueprintType)
 struct FLiveLinkOpenTrackIOStaticDuration
 {
 	GENERATED_BODY()
@@ -372,17 +402,17 @@ struct FLiveLinkOpenTrackIOStaticDuration
 	bool operator==(const FLiveLinkOpenTrackIOStaticDuration&) const = default;
 	bool operator!=(const FLiveLinkOpenTrackIOStaticDuration&) const = default;
 	
-	UPROPERTY()
+	UPROPERTY(EditAnywhere, BlueprintReadWrite, Category = "OpenTrackIO")
 	int32 Num = 1;
 
-	UPROPERTY()
-	uint32 Denom = 1; 
+	UPROPERTY(EditAnywhere, BlueprintReadWrite, Category = "OpenTrackIO")
+	int32 Denom = 1; 
 };
 
 /**
  * Static data from a parsed OpenTrackIO datagram. Per the spec this will be sent periodically from the source.
  */
-USTRUCT()
+USTRUCT(BlueprintType)
 struct FLiveLinkOpenTrackIOStaticData
 {
 	GENERATED_BODY()
@@ -397,20 +427,20 @@ struct FLiveLinkOpenTrackIOStaticData
 	bool operator==(const FLiveLinkOpenTrackIOStaticData&) const = default;
 	bool operator!=(const FLiveLinkOpenTrackIOStaticData&) const = default;
 	
-	UPROPERTY()
+	UPROPERTY(EditAnywhere, BlueprintReadWrite, Category = "OpenTrackIO")
 	FLiveLinkOpenTrackIOStaticDuration Duration;
 
-	UPROPERTY()
+	UPROPERTY(EditAnywhere, BlueprintReadWrite, Category = "OpenTrackIO")
 	FLiveLinkOpenTrackIOStaticCamera Camera;
 
-	UPROPERTY()
+	UPROPERTY(EditAnywhere, BlueprintReadWrite, Category = "OpenTrackIO")
 	FLiveLinkOpenTrackIOStaticLens Lens;
 
-	UPROPERTY()
+	UPROPERTY(EditAnywhere, BlueprintReadWrite, Category = "OpenTrackIO")
 	FLiveLinkOpenTrackIOStaticTracker Tracker;
 };
 
-USTRUCT()
+USTRUCT(BlueprintType)
 struct FLiveLinkOpenTrackIOTracker
 {
 	GENERATED_BODY()
@@ -426,19 +456,19 @@ struct FLiveLinkOpenTrackIOTracker
 	bool operator!=(const FLiveLinkOpenTrackIOTracker&) const = default;
 	
 	/** Non-blank string containing notes about tracking system */
-	UPROPERTY()
+	UPROPERTY(EditAnywhere, BlueprintReadWrite, Category = "OpenTrackIO")
 	FString Notes;
 
 	/** Boolean indicating whether tracking system is recording data. */
-	UPROPERTY()
+	UPROPERTY(EditAnywhere, BlueprintReadWrite, Category = "OpenTrackIO")
 	bool Recording = false;
 
 	/** Non-blank string describing the recording slate */
-	UPROPERTY()
+	UPROPERTY(EditAnywhere, BlueprintReadWrite, Category = "OpenTrackIO")
 	FString Slate;
 
 	/** Non-blank string describing status of tracking system */
-	UPROPERTY()
+	UPROPERTY(EditAnywhere, BlueprintReadWrite, Category = "OpenTrackIO")
 	FString Status;
 };
 
@@ -449,7 +479,7 @@ struct FLiveLinkOpenTrackIOTracker
  *
  * We use 64-bit unsigned values for storing seconds. 
  */
-USTRUCT()
+USTRUCT(BlueprintType)
 struct FLiveLinkOpenTrackIOTimestamp
 {
 	GENERATED_BODY()
@@ -465,15 +495,15 @@ struct FLiveLinkOpenTrackIOTimestamp
 	bool operator!=(const FLiveLinkOpenTrackIOTimestamp&) const = default;
 	
 	/** Although the schema specifies a 48 bit number. We allocate a full 64 bit integer here to store it. */
-	UPROPERTY()
-	uint64 Seconds = 0;
+	UPROPERTY(EditAnywhere, BlueprintReadWrite, Category = "OpenTrackIO")
+	int64 Seconds = 0;
 
 	/** A 32-bit unsigned number for nanoseconds */
-	UPROPERTY()
-	uint32 Nanoseconds = 0;
+	UPROPERTY(EditAnywhere, BlueprintReadWrite, Category = "OpenTrackIO")
+	int32 Nanoseconds = 0;
 };
 
-USTRUCT()
+USTRUCT(BlueprintType)
 struct FLiveLinkOpenTrackIOSynchronizationOffsets
 {
 	GENERATED_BODY()
@@ -488,17 +518,17 @@ struct FLiveLinkOpenTrackIOSynchronizationOffsets
 	bool operator==(const FLiveLinkOpenTrackIOSynchronizationOffsets&) const = default;
 	bool operator!=(const FLiveLinkOpenTrackIOSynchronizationOffsets&) const = default;
 	
-	UPROPERTY()
+	UPROPERTY(EditAnywhere, BlueprintReadWrite, Category = "OpenTrackIO")
 	float Translation = 0.0f;
 
-	UPROPERTY()
+	UPROPERTY(EditAnywhere, BlueprintReadWrite, Category = "OpenTrackIO")
 	float Rotation = 0.0f;
 
-	UPROPERTY()
+	UPROPERTY(EditAnywhere, BlueprintReadWrite, Category = "OpenTrackIO")
 	float LensEncoders = 0.0f;
 };
 
-USTRUCT()
+USTRUCT(BlueprintType)
 struct FLiveLinkOpenTrackIO_PTPLeaderPriorities
 {
 	GENERATED_BODY()
@@ -513,14 +543,14 @@ struct FLiveLinkOpenTrackIO_PTPLeaderPriorities
 	bool operator==(const FLiveLinkOpenTrackIO_PTPLeaderPriorities&) const = default;
 	bool operator!=(const FLiveLinkOpenTrackIO_PTPLeaderPriorities&) const = default;
 	
-	UPROPERTY()
+	UPROPERTY(EditAnywhere, BlueprintReadWrite, Category = "OpenTrackIO")
 	uint8 Priority1 = 0;
 
-	UPROPERTY()
+	UPROPERTY(EditAnywhere, BlueprintReadWrite, Category = "OpenTrackIO")
 	uint8 Priority2 = 0;
 };
 
-USTRUCT()
+USTRUCT(BlueprintType)
 struct FLiveLinkOpenTrackIO_PTP
 {
 	GENERATED_BODY()
@@ -536,7 +566,7 @@ struct FLiveLinkOpenTrackIO_PTP
 	bool operator!=(const FLiveLinkOpenTrackIO_PTP&) const = default;
 	
 	/** PTP Profile: Can be one the following values [IEEE Std 1588-2019, IEEE Std 802.1AS-2020, SMPTE ST2059-2:2021] */
-	UPROPERTY()
+	UPROPERTY(EditAnywhere, BlueprintReadWrite, Category = "OpenTrackIO")
 	FString Profile = TEXT("SMPTE ST2059-2:2021");
 
 	UPROPERTY(meta = (ClampMin = "0"))
@@ -547,25 +577,25 @@ struct FLiveLinkOpenTrackIO_PTP
 	 * 
 	 * (?:^[0-9a-f]{2}(?::[0-9a-f]{2}){5}$)|(?:^[0-9a-f]{2}(?:-[0-9a-f]{2}){5}$)
 	 */
-	UPROPERTY()
+	UPROPERTY(EditAnywhere, BlueprintReadWrite, Category = "OpenTrackIO")
 	FString LeaderIdentity; 
 
 	/**  Data structure for PTP synchronization priorities */
-	UPROPERTY()
+	UPROPERTY(EditAnywhere, BlueprintReadWrite, Category = "OpenTrackIO")
 	FLiveLinkOpenTrackIO_PTPLeaderPriorities LeaderPriorities;
 
-	UPROPERTY()
+	UPROPERTY(EditAnywhere, BlueprintReadWrite, Category = "OpenTrackIO")
 	float LeaderAccuracy = 0.0;
 
 	/** Enum string property [GNSS, Atomic clock, NTP] */
-	UPROPERTY()
+	UPROPERTY(EditAnywhere, BlueprintReadWrite, Category = "OpenTrackIO")
 	FString LeaderTimeSource = TEXT("NTP");
 
-	UPROPERTY()
+	UPROPERTY(EditAnywhere, BlueprintReadWrite, Category = "OpenTrackIO")
 	float MeanPathDelay = 0.0;
 
-	UPROPERTY()
-	uint32 Vlan = 0;
+	UPROPERTY(EditAnywhere, BlueprintReadWrite, Category = "OpenTrackIO")
+	int32 Vlan = 0;
 };
 
 /**
@@ -623,7 +653,7 @@ struct FLiveLinkOpenTrackIO_PTP
  *  - "ntp": The tracking device is locked to an NTP server
  *  
  */
-USTRUCT()
+USTRUCT(BlueprintType)
 struct FLiveLinkOpenTrackIOSynchronization
 {
 	GENERATED_BODY()
@@ -638,24 +668,24 @@ struct FLiveLinkOpenTrackIOSynchronization
 	bool operator==(const FLiveLinkOpenTrackIOSynchronization&) const = default;
 	bool operator!=(const FLiveLinkOpenTrackIOSynchronization&) const = default;
 	
-	UPROPERTY()
+	UPROPERTY(EditAnywhere, BlueprintReadWrite, Category = "OpenTrackIO")
 	bool Locked = false;
 
 	/** Enum as string value can be "genlock", "videoIn", "ptp", ... */
-	UPROPERTY()
+	UPROPERTY(EditAnywhere, BlueprintReadWrite, Category = "OpenTrackIO")
 	FString Source;
 
 	/** Synchonization rate. */
-	UPROPERTY()
+	UPROPERTY(EditAnywhere, BlueprintReadWrite, Category = "OpenTrackIO")
 	FLiveLinkOpenTrackIOFrameRate Frequency;
 
-	UPROPERTY()
+	UPROPERTY(EditAnywhere, BlueprintReadWrite, Category = "OpenTrackIO")
 	FLiveLinkOpenTrackIOSynchronizationOffsets Offsets;
 
-	UPROPERTY()
+	UPROPERTY(EditAnywhere, BlueprintReadWrite, Category = "OpenTrackIO")
 	bool Present = false;
 
-	UPROPERTY()
+	UPROPERTY(EditAnywhere, BlueprintReadWrite, Category = "OpenTrackIO")
 	FLiveLinkOpenTrackIO_PTP Ptp;
 };
 
@@ -669,7 +699,7 @@ struct FLiveLinkOpenTrackIOSynchronization
  * division of the frame, e.g. interlaced frames have two sub-frames,
  * one per field.
  */
-USTRUCT()
+USTRUCT(BlueprintType)
 struct FLiveLinkOpenTrackIOTimecode
 {
 	GENERATED_BODY()
@@ -696,26 +726,26 @@ struct FLiveLinkOpenTrackIOTimecode
 		return FQualifiedFrameTime(GetTimecode(), FrameRate.GetFrameRate());
 	}
 	
-	UPROPERTY()	
+	UPROPERTY(EditAnywhere, BlueprintReadWrite, Category = "OpenTrackIO")	
 	int32 Hours = 0;
 
-	UPROPERTY()
+	UPROPERTY(EditAnywhere, BlueprintReadWrite, Category = "OpenTrackIO")
 	int32 Minutes = 0;
 
-	UPROPERTY()
+	UPROPERTY(EditAnywhere, BlueprintReadWrite, Category = "OpenTrackIO")
 	int32 Seconds = 0;
 
-	UPROPERTY()
+	UPROPERTY(EditAnywhere, BlueprintReadWrite, Category = "OpenTrackIO")
 	int32 Frames = 0;
 
-	UPROPERTY()
+	UPROPERTY(EditAnywhere, BlueprintReadWrite, Category = "OpenTrackIO")
 	FLiveLinkOpenTrackIOFrameRate FrameRate;
 	
-	UPROPERTY()
+	UPROPERTY(EditAnywhere, BlueprintReadWrite, Category = "OpenTrackIO")
 	int32 SubFrame = 0;
 };
 
-USTRUCT()
+USTRUCT(BlueprintType)
 struct FLiveLinkOpenTrackIOTiming
 {
 	GENERATED_BODY()
@@ -735,11 +765,11 @@ struct FLiveLinkOpenTrackIOTiming
 	 *  inherent (’external’) timing, or whether the transport mechanism lacks inherent timing
 	 *  and so the sample must contain a PTP timestamp itself (’internal’) to carry timing information.
 	 **/ 
-	UPROPERTY()
+	UPROPERTY(EditAnywhere, BlueprintReadWrite, Category = "OpenTrackIO")
 	FString Mode;
 
 	/** PTP timestamp of the data recording instant, provided for convenience during playback of e.g. pre-recorded tracking data. */
-	UPROPERTY()
+	UPROPERTY(EditAnywhere, BlueprintReadWrite, Category = "OpenTrackIO")
 	FLiveLinkOpenTrackIOTimestamp RecordedTimestamp;
 
 	/**
@@ -747,30 +777,30 @@ struct FLiveLinkOpenTrackIOTiming
 	 * represented as e.g. 30000/1001. In a variable rate system this should be estimated from the last
 	 * sample delta time.
 	 */
-	UPROPERTY()
+	UPROPERTY(EditAnywhere, BlueprintReadWrite, Category = "OpenTrackIO")
 	FLiveLinkOpenTrackIOFrameRate SampleRate;
 
 	/** 
 	 * PTP timestamp of the data recording instant, provided for convenience during playback of e.g. pre-recorded tracking data.
 	 * The timestamp comprises a 48-bit unsigned integer (seconds), a 32-bit unsigned integer (nanoseconds)
 	 */
-	UPROPERTY()
+	UPROPERTY(EditAnywhere, BlueprintReadWrite, Category = "OpenTrackIO")
 	FLiveLinkOpenTrackIOTimestamp SampleTimestamp;
 
 	/** Integer incrementing with each sample. */
-	UPROPERTY()
-	uint32 SequenceNumber = 0;
+	UPROPERTY(EditAnywhere, BlueprintReadWrite, Category = "OpenTrackIO")
+	int32 SequenceNumber = 0;
 
 	/** Synchronization object for the timing. */
-	UPROPERTY()
+	UPROPERTY(EditAnywhere, BlueprintReadWrite, Category = "OpenTrackIO")
 	FLiveLinkOpenTrackIOSynchronization Synchronization;
 
 	/** SMTPE timecode of the sample. */
-	UPROPERTY()
+	UPROPERTY(EditAnywhere, BlueprintReadWrite, Category = "OpenTrackIO")
 	FLiveLinkOpenTrackIOTimecode Timecode;
 };
 
-USTRUCT()
+USTRUCT(BlueprintType)
 struct FLiveLinkOpenTrackIOLens_DistortionOffset
 {
 	GENERATED_BODY()
@@ -790,10 +820,10 @@ struct FLiveLinkOpenTrackIOLens_DistortionOffset
 		return FVector2D(X,Y);
 	}
 	
-	UPROPERTY()
+	UPROPERTY(EditAnywhere, BlueprintReadWrite, Category = "OpenTrackIO")
 	float X = 0;
 
-	UPROPERTY()
+	UPROPERTY(EditAnywhere, BlueprintReadWrite, Category = "OpenTrackIO")
 	float Y = 0;
 };
 
@@ -809,7 +839,7 @@ struct FLiveLinkOpenTrackIOLens_DistortionOffset
  * Iris:    0=open         1=closed
  * Zoom:    0=wide angle   1=telephoto
  */
-USTRUCT()
+USTRUCT(BlueprintType)
 struct FLiveLinkOpenTrackIOLens_Encoders
 {
 	GENERATED_BODY()
@@ -824,13 +854,13 @@ struct FLiveLinkOpenTrackIOLens_Encoders
 	bool operator==(const FLiveLinkOpenTrackIOLens_Encoders&) const = default;
 	bool operator!=(const FLiveLinkOpenTrackIOLens_Encoders&) const = default;
 	
-	UPROPERTY()
+	UPROPERTY(EditAnywhere, BlueprintReadWrite, Category = "OpenTrackIO")
 	float Focus = 0;
 
-	UPROPERTY()
+	UPROPERTY(EditAnywhere, BlueprintReadWrite, Category = "OpenTrackIO")
 	float Iris = 0;
 
-	UPROPERTY()
+	UPROPERTY(EditAnywhere, BlueprintReadWrite, Category = "OpenTrackIO")
 	float Zoom = 0;
 };
 
@@ -840,7 +870,7 @@ struct FLiveLinkOpenTrackIOLens_Encoders
  * These values are dependent on encoder resolution and before any
  * homing / ranging has taken place.
  */
-USTRUCT()
+USTRUCT(BlueprintType)
 struct FLiveLinkOpenTrackIOLens_RawEncoders
 {
 	GENERATED_BODY()
@@ -855,20 +885,20 @@ struct FLiveLinkOpenTrackIOLens_RawEncoders
 	bool operator==(const FLiveLinkOpenTrackIOLens_RawEncoders&) const = default;
 	bool operator!=(const FLiveLinkOpenTrackIOLens_RawEncoders&) const = default;
 	
-	UPROPERTY()
-	uint32 Focus = 0;
+	UPROPERTY(EditAnywhere, BlueprintReadWrite, Category = "OpenTrackIO")
+	int32 Focus = 0;
 
-	UPROPERTY()
-	uint32 Iris = 0;
+	UPROPERTY(EditAnywhere, BlueprintReadWrite, Category = "OpenTrackIO")
+	int32 Iris = 0;
 
-	UPROPERTY()
-	uint32 Zoom = 0;
+	UPROPERTY(EditAnywhere, BlueprintReadWrite, Category = "OpenTrackIO")
+	int32 Zoom = 0;
 };
 
 /**
  * Coefficients for calculating the exposure fall-off (vignetting) of a lens.
  */
-USTRUCT()
+USTRUCT(BlueprintType)
 struct FLiveLinkOpenTrackIOLens_ExposureFalloff
 {
 	GENERATED_BODY()
@@ -883,17 +913,17 @@ struct FLiveLinkOpenTrackIOLens_ExposureFalloff
 	bool operator==(const FLiveLinkOpenTrackIOLens_ExposureFalloff&) const = default;
 	bool operator!=(const FLiveLinkOpenTrackIOLens_ExposureFalloff&) const = default;
 	
-	UPROPERTY()
+	UPROPERTY(EditAnywhere, BlueprintReadWrite, Category = "OpenTrackIO")
 	float A1 = 0;
 
-	UPROPERTY()
+	UPROPERTY(EditAnywhere, BlueprintReadWrite, Category = "OpenTrackIO")
 	float A2 = 0;
 
-	UPROPERTY()
+	UPROPERTY(EditAnywhere, BlueprintReadWrite, Category = "OpenTrackIO")
 	float A3 = 0;
 };
 
-USTRUCT()
+USTRUCT(BlueprintType)
 struct FLiveLinkOpenTrackIOLens_DistortionCoeff
 {
 	GENERATED_BODY()
@@ -908,17 +938,17 @@ struct FLiveLinkOpenTrackIOLens_DistortionCoeff
 	bool operator==(const FLiveLinkOpenTrackIOLens_DistortionCoeff&) const = default;
 	bool operator!=(const FLiveLinkOpenTrackIOLens_DistortionCoeff&) const = default;
 	
-	UPROPERTY()
+	UPROPERTY(EditAnywhere, BlueprintReadWrite, Category = "OpenTrackIO")
 	TArray<float> Radial;
 
-	UPROPERTY()
+	UPROPERTY(EditAnywhere, BlueprintReadWrite, Category = "OpenTrackIO")
 	TArray<float> Tangential;
 
-	UPROPERTY()
+	UPROPERTY(EditAnywhere, BlueprintReadWrite, Category = "OpenTrackIO")
 	float Overscan = 0.0;
 };
 
-USTRUCT()
+USTRUCT(BlueprintType)
 struct FLiveLinkOpenTrackIOLens
 {
 	GENERATED_BODY()
@@ -938,7 +968,7 @@ struct FLiveLinkOpenTrackIOLens
 	 * of and how these characteristics are to be applied to a virtual camera would require
 	 * negotiation between a particular producer and consumer.
 	 */
-	UPROPERTY()
+	UPROPERTY(EditAnywhere, BlueprintReadWrite, Category = "OpenTrackIO")
 	TArray<float> Custom;
 
 	/**
@@ -947,17 +977,17 @@ struct FLiveLinkOpenTrackIOLens
 	 * distortion model. Typical values for ’model’ include "Brown-Conrady D-U" when mapping distorted to undistorted coordinates, and "Brown-Conrady
 	 * U-D" when mapping undistorted to undistorted coordinates. If not provided, the default model is "Brown-Conrady D-U".
 	 */
-	UPROPERTY()
+	UPROPERTY(EditAnywhere, BlueprintReadWrite, Category = "OpenTrackIO")
 	TArray<FLiveLinkOpenTrackIOLens_DistortionCoeff> Distortion;
 
 	/**
 	 *  Offset in x and y of the centre of distortion of the virtual camera in millimeter
 	 */
-	UPROPERTY()
+	UPROPERTY(EditAnywhere, BlueprintReadWrite, Category = "OpenTrackIO")
 	FLiveLinkOpenTrackIOLens_DistortionOffset DistortionOffset;
 
 	/** Normalised real numbers (0-1) for focus, iris and zoom */
-	UPROPERTY()
+	UPROPERTY(EditAnywhere, BlueprintReadWrite, Category = "OpenTrackIO")
 	FLiveLinkOpenTrackIOLens_Encoders Encoders;
 
 	/**
@@ -967,37 +997,37 @@ struct FLiveLinkOpenTrackIOLens
 	 * Measured in meters as in a render engine it is often applied in the
 	 * virtual camera’s transform chain.
 	 */
-	UPROPERTY()
+	UPROPERTY(EditAnywhere, BlueprintReadWrite, Category = "OpenTrackIO")
 	float  EntrancePupilOffset = 0;
 
 	/**
 	 * Coefficients for calculating the exposure fall-off (vignetting) of a lens
 	 */
-	UPROPERTY()
+	UPROPERTY(EditAnywhere, BlueprintReadWrite, Category = "OpenTrackIO")
 	FLiveLinkOpenTrackIOLens_ExposureFalloff ExposureFalloff;
 
 	/** The linear f-number of the lens, equal to the focal length divided by the diameter of the entrance pupil. */
-	UPROPERTY()
+	UPROPERTY(EditAnywhere, BlueprintReadWrite, Category = "OpenTrackIO")
 	float FStop = 2.8f;
 
 	/** Distance between the pinhole and the image plane in the simple CGI pinhole camera model. (millimeters) */
-	UPROPERTY()
+	UPROPERTY(EditAnywhere, BlueprintReadWrite, Category = "OpenTrackIO")
 	float PinholeFocalLength = 35.0f;
 
 	/** Focus distance/position of the lens in meters. */ 
-	UPROPERTY()
+	UPROPERTY(EditAnywhere, BlueprintReadWrite, Category = "OpenTrackIO")
 	float FocusDistance = 10.0f;
 
 	/**  Offset in x and y of the centre of perspective projection of the virtual camera */
-	UPROPERTY()
+	UPROPERTY(EditAnywhere, BlueprintReadWrite, Category = "OpenTrackIO")
 	FLiveLinkOpenTrackIOLens_DistortionOffset ProjectionOffset;
 
 	/** Raw encoder values. */
-	UPROPERTY()
+	UPROPERTY(EditAnywhere, BlueprintReadWrite, Category = "OpenTrackIO")
 	FLiveLinkOpenTrackIOLens_RawEncoders RawEncoders;
 
 	/** T-stop: Linear t-number of the lens, equal to the F-number of the lens divided by the square root of the transmittance of the lens. */
-	UPROPERTY()
+	UPROPERTY(EditAnywhere, BlueprintReadWrite, Category = "OpenTrackIO")
 	float TStop = 0.0;
 };
 
@@ -1005,11 +1035,16 @@ struct FLiveLinkOpenTrackIOLens
  * Name of the protocol in which the sample is being employed, and
  * version of that protocol
  */
-USTRUCT()
+USTRUCT(BlueprintType)
 struct FLiveLinkOpenTrackIOProtocol
 {
 	GENERATED_BODY()
 
+	FLiveLinkOpenTrackIOProtocol()
+	{
+		Version = TArray<int8>{1,0,0};
+	};
+	
 	/** Returns true if the struct is equivalent to the default struct values. */
 	bool IsDefault() const
 	{
@@ -1019,7 +1054,7 @@ struct FLiveLinkOpenTrackIOProtocol
 	
 	bool operator==(const FLiveLinkOpenTrackIOProtocol&) const = default;
 	bool operator!=(const FLiveLinkOpenTrackIOProtocol&) const = default;
-
+	
 	bool IsSupported() const
 	{
 		/** Validate that we support this version of OpenTrackIO. We assume that patch release are compatible. */
@@ -1029,20 +1064,40 @@ struct FLiveLinkOpenTrackIOProtocol
 			&& Version[1] == 0;
 	}
 
-	UPROPERTY()
+	/** The name of the protocol. If you want to extract the full version string then use GetVersionString().*/
+	UPROPERTY(BlueprintReadOnly, Category = "OpenTrackIO")
 	FString Name;
 
-	// 3 numbers that represent the protocol version.  
+	/** Three numbers that represent the protocol version.  This is not accessible by Blueprints because int8 is not supported.  */
 	UPROPERTY()
 	TArray<int8> Version;
 };
 
+UCLASS(meta=(BlueprintThreadSafe, ScriptName = "OpenTrackIO"), MinimalAPI)
+class ULiveLinkOpenTrackIOLibrary : public UBlueprintFunctionLibrary
+{
+	GENERATED_BODY()
+
+public:
+	
+	/** A string representing the protocol name plus version. */
+	UFUNCTION(BlueprintPure, Category = "OpenTrackIO")
+	static FString GetVersionString(const FLiveLinkOpenTrackIOProtocol& Protocol)
+	{
+		int Major = Protocol.Version.Num() > 0 ? Protocol.Version[0] : 1;
+		int Minor = Protocol.Version.Num() > 1 ? Protocol.Version[1] : 0;
+		int Patch = Protocol.Version.Num() > 2 ? Protocol.Version[2] : 0;
+		return FString::Printf(TEXT("%s %d.%d.%d"),	*Protocol.Name, Major, Minor, Patch);
+	}
+};
+
+
 /**
  * Position of stage origin in global ENU and geodetic coordinates
  * (E, N, U, lat0, lon0, h0). Note this may be dynamic if the stage is
  * inside a moving vehicle. Units represented in meters. 
  */
-USTRUCT()
+USTRUCT(BlueprintType)
 struct FLiveLinkOpenTrackIOGlobalStage
 {
 	GENERATED_BODY()
@@ -1057,26 +1112,26 @@ struct FLiveLinkOpenTrackIOGlobalStage
 	bool operator==(const FLiveLinkOpenTrackIOGlobalStage&) const = default;
 	bool operator!=(const FLiveLinkOpenTrackIOGlobalStage&) const = default;
 
-	UPROPERTY()
+	UPROPERTY(EditAnywhere, BlueprintReadWrite, Category = "OpenTrackIO")
 	float E = 0;
 
-	UPROPERTY()
+	UPROPERTY(EditAnywhere, BlueprintReadWrite, Category = "OpenTrackIO")
 	float N = 0;
 
-	UPROPERTY()
+	UPROPERTY(EditAnywhere, BlueprintReadWrite, Category = "OpenTrackIO")
 	float U = 0;
 
-	UPROPERTY()
+	UPROPERTY(EditAnywhere, BlueprintReadWrite, Category = "OpenTrackIO")
 	float Lat0 = 0;
 
-	UPROPERTY()
+	UPROPERTY(EditAnywhere, BlueprintReadWrite, Category = "OpenTrackIO")
 	float Lon0 = 0;
 
-	UPROPERTY()
+	UPROPERTY(EditAnywhere, BlueprintReadWrite, Category = "OpenTrackIO")
 	float H0 = 0;
 };
 
-USTRUCT()
+USTRUCT(BlueprintType)
 struct FLiveLinkOpenTrackIO_Rotator
 {
 	GENERATED_BODY()
@@ -1097,17 +1152,17 @@ struct FLiveLinkOpenTrackIO_Rotator
 		return FRotator(Tilt, -Pan, Roll);
 	}
 
-	UPROPERTY()
+	UPROPERTY(EditAnywhere, BlueprintReadWrite, Category = "OpenTrackIO")
 	float Pan = 0;
 
-	UPROPERTY()
+	UPROPERTY(EditAnywhere, BlueprintReadWrite, Category = "OpenTrackIO")
 	float Tilt = 0;
 
-	UPROPERTY()
+	UPROPERTY(EditAnywhere, BlueprintReadWrite, Category = "OpenTrackIO")
 	float Roll = 0;
 };
 
-USTRUCT()
+USTRUCT(BlueprintType)
 struct FLiveLinkOpenTrackIO_XYZ
 {
 	GENERATED_BODY()
@@ -1140,17 +1195,16 @@ struct FLiveLinkOpenTrackIO_XYZ
 		return FVector(Y, X, Z);
 	}
 
-	UPROPERTY()
+	UPROPERTY(EditAnywhere, BlueprintReadWrite, Category = "OpenTrackIO")
 	float X = 0;
-
-	UPROPERTY()
+	UPROPERTY(EditAnywhere, BlueprintReadWrite, Category = "OpenTrackIO")
 	float Y = 0;
 
-	UPROPERTY()
+	UPROPERTY(EditAnywhere, BlueprintReadWrite, Category = "OpenTrackIO")
 	float Z = 0;
 };
 
-USTRUCT()
+USTRUCT(BlueprintType)
 struct FLiveLinkOpenTrackIOTransform
 {
 	GENERATED_BODY()
@@ -1176,29 +1230,73 @@ struct FLiveLinkOpenTrackIOTransform
 	}
 	
 	/** Translation part of a OpenTrackIO xform */
-	UPROPERTY()
+	UPROPERTY(EditAnywhere, BlueprintReadWrite, Category = "OpenTrackIO")
 	FLiveLinkOpenTrackIO_XYZ Translation;
 	
 	/** Rotation part of a OpenTrackIO xform */
-	UPROPERTY()
+	UPROPERTY(EditAnywhere, BlueprintReadWrite, Category = "OpenTrackIO")
 	FLiveLinkOpenTrackIO_Rotator Rotation;
 
 	/** Scale part of a OpenTrackIO xform */
-	UPROPERTY()
+	UPROPERTY(EditAnywhere, BlueprintReadWrite, Category = "OpenTrackIO")
 	FLiveLinkOpenTrackIO_XYZ Scale = FLiveLinkOpenTrackIO_XYZ(1,1,1);
 
 	/** OpenTrackIO identifier (can be empty) */
-	UPROPERTY()
+	UPROPERTY(EditAnywhere, BlueprintReadWrite, Category = "OpenTrackIO")
 	FString Id;
 };
 
+USTRUCT(BlueprintType)
+struct FLiveLinkOpenTrackIOCustomDataField
+{
+	GENERATED_BODY()
+
+	/** Returns true if the struct is equivalent to the default struct values. */
+	bool IsDefault() const
+	{
+		static FLiveLinkOpenTrackIOCustomDataField DefaultValue;
+		return DefaultValue == *this;
+	}
+	
+	bool operator==(const FLiveLinkOpenTrackIOCustomDataField&) const = default;
+	bool operator!=(const FLiveLinkOpenTrackIOCustomDataField&) const = default;
+
+	/** Meta data is inserted into Live Link Frame data using a Key, Value pair. This is the key part of that pair. */
+	UPROPERTY(EditAnywhere, BlueprintReadWrite, Category = "OpenTrackIO")
+	FString Key;
+
+	/** Meta data is inserted into Live Link Frame data using a Key, Value pair. This is the value part of that pair. */
+	UPROPERTY(EditAnywhere, BlueprintReadWrite, Category = "OpenTrackIO")
+	FString Value;
+};
+
+USTRUCT(BlueprintType)
+struct FLiveLinkOpenTrackIOCustomData
+{
+	GENERATED_BODY()
+
+	/** Returns true if the struct is equivalent to the default struct values. */
+	bool IsDefault() const
+	{
+		static FLiveLinkOpenTrackIOCustomData DefaultValue;
+		return DefaultValue == *this;
+	}
+	
+	bool operator==(const FLiveLinkOpenTrackIOCustomData&) const = default;
+	bool operator!=(const FLiveLinkOpenTrackIOCustomData&) const = default;
+
+	/** A list of string based key/value pairs that can be read and applied to the Live Link data stream. */
+	UPROPERTY(EditAnywhere, BlueprintReadWrite, Category = "OpenTrackIO")
+	TArray<FLiveLinkOpenTrackIOCustomDataField> LiveLinkMetaData; 
+};
+
 /**
  * Data from a parsed header that was read by other CBOR or JSON format.
  *
  * Schema can be found here: https://www.opentrackio.org/schema.json
  * 
  */
-USTRUCT()
+USTRUCT(BlueprintType)
 struct FLiveLinkOpenTrackIOData
 {
 	GENERATED_BODY()
@@ -1213,33 +1311,33 @@ struct FLiveLinkOpenTrackIOData
 	bool operator==(const FLiveLinkOpenTrackIOData&) const = default;
 	bool operator!=(const FLiveLinkOpenTrackIOData&) const = default;
 	
-	UPROPERTY()
+	UPROPERTY(EditAnywhere, BlueprintReadWrite, Category = "OpenTrackIO")
 	FLiveLinkOpenTrackIOStaticData Static;
 
-	UPROPERTY()
+	UPROPERTY(EditAnywhere, BlueprintReadWrite, Category = "OpenTrackIO")
 	FLiveLinkOpenTrackIOTracker Tracker;
 
-	UPROPERTY()
+	UPROPERTY(EditAnywhere, BlueprintReadWrite, Category = "OpenTrackIO")
 	FLiveLinkOpenTrackIOTiming Timing;
 
-	UPROPERTY()
+	UPROPERTY(EditAnywhere, BlueprintReadWrite, Category = "OpenTrackIO")
 	FLiveLinkOpenTrackIOLens Lens;
 
-	UPROPERTY()
+	UPROPERTY(EditAnywhere, BlueprintReadWrite, Category = "OpenTrackIO")
 	FLiveLinkOpenTrackIOProtocol Protocol;
 
 	/**
 	 * URN serving as unique identifier of the sample in which data is being transported.
 	 * Pattern -> ^urn:uuid:[0-9a-f]{8}-[0-9a-f]{4}-[0-9a-f]{4}-[0-9a-f]{4}-[0-9a-f]{12}$
 	 */
-	UPROPERTY()
+	UPROPERTY(EditAnywhere, BlueprintReadWrite, Category = "OpenTrackIO")
 	FString SampleId;
 
 	/**
 	 * URN serving as unique identifier of the source from which data is being transported.
 	 * Pattern -> ^urn:uuid:[0-9a-f]{8}-[0-9a-f]{4}-[0-9a-f]{4}-[0-9a-f]{4}-[0-9a-f]{12}$
 	 */
-	UPROPERTY()
+	UPROPERTY(EditAnywhere, BlueprintReadWrite, Category = "OpenTrackIO")
 	FString SourceId;
 
 	/**
@@ -1247,18 +1345,18 @@ struct FLiveLinkOpenTrackIOData
 	 * data is being transported. This is most important in the case where a source
 	 * is producing multiple streams of samples.
 	 */
-	UPROPERTY()
-	uint64 SourceNumber = 0;
+	UPROPERTY(EditAnywhere, BlueprintReadWrite, Category = "OpenTrackIO")
+	int64 SourceNumber = 0;
 
 	/**
 	 * List of sampleId properties of samples related to this sample. The
 	 * existence of a sample with a given sampleId is not guaranteed.
 	 */
-	UPROPERTY()
+	UPROPERTY(EditAnywhere, BlueprintReadWrite, Category = "OpenTrackIO")
 	TArray<FString> RelatedSampleIds;
 
 	/** Position of stage origin in global ENU and geodetic coordinates */
-	UPROPERTY()
+	UPROPERTY(EditAnywhere, BlueprintReadWrite, Category = "OpenTrackIO")
 	FLiveLinkOpenTrackIOGlobalStage GlobalStage;
 
 	 /**
@@ -1271,8 +1369,27 @@ struct FLiveLinkOpenTrackIOData
 	  * lock does not present the physical challenges of a robotic system. Conversion to and from quarternions is trivial with an acceptable loss of
 	  * precision.
 	  */ 
-	UPROPERTY()
+	UPROPERTY(EditAnywhere, BlueprintReadWrite, Category = "OpenTrackIO")
 	TArray<FLiveLinkOpenTrackIOTransform> Transforms;
-	
+
+	/**
+	 * Typed property to allow users to add custom data to the incoming OpenTrackIO data. Our implementation only supports a field called LiveLinkMetaData.
+	 * For example,
+	 * 
+	 * 	"custom" : {
+	 *     "liveLinkMetaData" : [
+	 *			{
+	 *				"key": "mykey",
+	 *				"value": "myvalue"
+	 *			},
+	 *			{
+	 *				"key": "otherKey",
+	 *				"value": "otherValue"
+	 *			}
+	 *		]
+	 *  }
+	 * */
+	UPROPERTY(EditAnywhere, BlueprintReadWrite, Category = "OpenTrackIO")
+	FLiveLinkOpenTrackIOCustomData Custom;
 };
 
diff --git a/Engine/Plugins/VirtualProduction/LiveLinkOpenTrackIO/Source/LiveLinkOpenTrackIO/Private/Tests/OpenTrackReaderTests.spec.cpp b/Engine/Plugins/VirtualProduction/LiveLinkOpenTrackIO/Source/LiveLinkOpenTrackIO/Private/Tests/OpenTrackReaderTests.spec.cpp
--- a/Engine/Plugins/VirtualProduction/LiveLinkOpenTrackIO/Source/LiveLinkOpenTrackIO/Private/Tests/OpenTrackReaderTests.spec.cpp
+++ b/Engine/Plugins/VirtualProduction/LiveLinkOpenTrackIO/Source/LiveLinkOpenTrackIO/Private/Tests/OpenTrackReaderTests.spec.cpp
@@ -32,43 +32,31 @@ BEGIN_DEFINE_SPEC(FOpenTrackIOReaderTests,
 	};
 
 	const TArray<FOTIO_DataValidation> CborValidationData = []() -> TArray<FOTIO_DataValidation> {
-		return { {.SequenceNumber =  2,
-					 .Checksum = 0x4532},
-				 {.SequenceNumber =  3,
-					 .Checksum = 0xDFC0},
-				 {.SequenceNumber =  4,
-				     .Checksum = 0xDCA4},
-				 {.SequenceNumber =  5,
-				     .Checksum = 0xCD54},
-				 {.SequenceNumber =  6,
-				     .Checksum = 0xCEB},
-				 {.SequenceNumber =  7,
-				     .Checksum = 0x647},
-				 {.SequenceNumber =  8,
-				     .Checksum = 0xE29D}};
+		return {
+		{ .SequenceNumber = 2, .Checksum = 0x86C }
+		,{ .SequenceNumber = 3, .Checksum = 0xA587 }
+		,{ .SequenceNumber = 4, .Checksum = 0xF596 }
+		,{ .SequenceNumber = 5, .Checksum = 0x46A5 }
+		,{ .SequenceNumber = 6, .Checksum = 0x96B4 }
+		,{ .SequenceNumber = 7, .Checksum = 0xE6C3 }
+		,{ .SequenceNumber = 8, .Checksum = 0x12CC }
+		,{ .SequenceNumber = 9, .Checksum = 0x3DD5 }
+		,{ .SequenceNumber = 10, .Checksum = 0x68DE }
+		,{ .SequenceNumber = 11, .Checksum = 0x93E7 }
+		,{ .SequenceNumber = 12, .Checksum = 0xBEF0 }};
 	}();
 	const TArray<FOTIO_DataValidation> JsonValidationData = []() -> TArray<FOTIO_DataValidation> {
-		return { 
-			{.SequenceNumber = 2,
-			 .Checksum =  0xB728},
-			{.SequenceNumber = 3,
-			 .Checksum =  0x8308},
-			{.SequenceNumber = 4,
-			 .Checksum =  0x239},
-			{.SequenceNumber = 5,
-			 .Checksum =  0xCE19},
-			{.SequenceNumber = 6,
-			 .Checksum =  0x4D4A},
-			{.SequenceNumber = 7,
-			 .Checksum =  0x192A},
-			{.SequenceNumber = 8,
-			 .Checksum =  0x985B},
-			{.SequenceNumber = 9,
-			 .Checksum =  0x643B},
-			{.SequenceNumber = 10,
-			 .Checksum =  0xE36C},
-			{.SequenceNumber = 11,
-			.Checksum =  0x9793}};
+		return {
+		{ .SequenceNumber = 2, .Checksum = 0xDA9C }
+		,{ .SequenceNumber = 3, .Checksum = 0xA293 }
+		,{ .SequenceNumber = 4, .Checksum = 0x7CA5 }
+		,{ .SequenceNumber = 5, .Checksum = 0x449C }
+		,{ .SequenceNumber = 6, .Checksum = 0x1EAE }
+		,{ .SequenceNumber = 7, .Checksum = 0xE5A5 }
+		,{ .SequenceNumber = 8, .Checksum = 0xBFB7 }
+		,{ .SequenceNumber = 9, .Checksum = 0x87AE }
+		,{ .SequenceNumber = 10, .Checksum = 0x61C0 }
+		,{ .SequenceNumber = 11, .Checksum = 0xA608 }};
 	}();
 END_DEFINE_SPEC(FOpenTrackIOReaderTests)
 
diff --git a/Engine/Source/Runtime/LiveLinkInterface/Public/Roles/LiveLinkCameraTypes.h b/Engine/Source/Runtime/LiveLinkInterface/Public/Roles/LiveLinkCameraTypes.h
--- a/Engine/Source/Runtime/LiveLinkInterface/Public/Roles/LiveLinkCameraTypes.h
+++ b/Engine/Source/Runtime/LiveLinkInterface/Public/Roles/LiveLinkCameraTypes.h
@@ -58,6 +58,10 @@ public:
 	//Set to false to force the camera to disable depth of field
 	UPROPERTY(EditAnywhere, BlueprintReadWrite, Category = "LiveLink")
 	bool bIsDepthOfFieldSupported = true;
+
+	// Specifies if the live link subject will send dynamic changes to the film back.
+	UPROPERTY(EditAnywhere, BlueprintReadWrite, Category = "LiveLink")
+	bool bIsDynamicFilmbackSupported = false;
 };
 
 /**
@@ -91,6 +95,15 @@ struct FLiveLinkCameraFrameData : public FLiveLinkTransformFrameData
 	// ProjectionMode of the camera (Perspective, Orthographic, etc...)
 	UPROPERTY(EditAnywhere, BlueprintReadWrite, Category = "LiveLink", Interp)
 	ELiveLinkCameraProjectionMode ProjectionMode = ELiveLinkCameraProjectionMode::Perspective;
+
+	// Used with CinematicCamera. Values greater than 0 will be applied if bIsDynamicFilmbackSupported is true on static data.
+	UPROPERTY(EditAnywhere, BlueprintReadWrite, Category = "LiveLink")
+	float FilmBackWidth = -1.0f;
+
+	// Used with CinematicCamera. Values greater than 0 will be applied if bIsDynamicFilmbackSupported is true on static data.
+	UPROPERTY(EditAnywhere, BlueprintReadWrite, Category = "LiveLink")
+	float FilmBackHeight = -1.0f;
+
 };
 
 /**
