OpenTrackIO: Relaxed requirement to send static camera and lens.

Also fixed a regression of not pushing the frame data with the default subject name because it was trying to re-derive it from the data.

#rnx

diff --git a/Engine/Plugins/VirtualProduction/LiveLinkOpenTrackIO/Source/LiveLinkOpenTrackIO/Private/LiveLinkOpenTrackIOSource.cpp b/Engine/Plugins/VirtualProduction/LiveLinkOpenTrackIO/Source/LiveLinkOpenTrackIO/Private/LiveLinkOpenTrackIOSource.cpp
--- a/Engine/Plugins/VirtualProduction/LiveLinkOpenTrackIO/Source/LiveLinkOpenTrackIO/Private/LiveLinkOpenTrackIOSource.cpp
+++ b/Engine/Plugins/VirtualProduction/LiveLinkOpenTrackIO/Source/LiveLinkOpenTrackIO/Private/LiveLinkOpenTrackIOSource.cpp
@@ -115,19 +115,6 @@ FIPv4Endpoint MakeMulticastAddressFromSourceNumber(uint8 InSourceNumber, uint16
 	return FIPv4Endpoint(FIPv4Address(235, 135, 1, SourceNumber), InPortNumber);
 }
 
-/**
- * Makes an interim camera label based on the OpenTrackIO Source Number.
- * Useful when the producer hasn't sent static camera data yet.
- *
- * @param InSourceNumber OpenTrackIO Source Number to insert into the subject name.
- * 
- * @return FString In the form "OpenTrackIO_<SourceNumber>".
- */
-FString MakeInterimCameraLabelFromSourceNumber(uint8 InSourceNumber)
-{
-	return FString::Printf(TEXT("OpenTrackIO_%03u"), InSourceNumber);
-}
-
 }
 
 FLiveLinkOpenTrackIOSource::FLiveLinkOpenTrackIOSource(FLiveLinkOpenTrackIOConnectionSettings InConnectionSettings)
@@ -356,40 +343,31 @@ void FLiveLinkOpenTrackIOSource::PushDataToLiveLink_AnyThread(const FLiveLinkOpe
 		Client->PushSubjectStaticData_AnyThread({ SourceGuid, SubjectName }, ULiveLinkLensRole::StaticClass(), MoveTemp(*StaticData));
 
 		// Remove the interim live link subject, if any.
-		if (!DefaultSubjectName.IsNone())
+		if (!Cache->DefaultSubjectName.IsNone())
 		{
-			Client->RemoveSubject_AnyThread({ SourceGuid, DefaultSubjectName });
-			DefaultSubjectName = NAME_None;
+			Client->RemoveSubject_AnyThread({ SourceGuid, Cache->DefaultSubjectName });
+			Cache->DefaultSubjectName = NAME_None;
 		}
 	}
 	else
 	{
-		// If no static camera data has arrived, we add an interim one named after the OpenTrackIO Source Number to support producers
-		// that do not send this information. If valid static camera data arrives later, the interim live link subject will be removed
+		// If no static camera data has arrived, we add an interim live link subject named after the OpenTrackIO Source Number
+		// to support producers that do not send this information. If valid static camera data arrives later, it will be removed
 		// and replaced with the new one.
-		if (!Cache->StaticCameras.Num())
+		if (!Cache->StaticCameras.Num() && Cache->DefaultSubjectName.IsNone())
 		{
-			using namespace UE::LiveLinkOpenTrackIO::Private;
+			StaticData = Cache->MakeStaticData(InData);
+			check(StaticData.IsSet());
 
-			FLiveLinkOpenTrackIOData InterimCameraData = InData;
+			// This is where the interim subject name is created and stored in the Cache's DefaultSubjectName.
+			Cache->DefaultSubjectName = FName(FString::Printf(TEXT("OpenTrackIO_%03u"), ConnectionSettings.SourceNumber));
 
-			InterimCameraData.Static.Camera.Label = MakeInterimCameraLabelFromSourceNumber(ConnectionSettings.SourceNumber);
-			InterimCameraData.Static.Camera.Make = "N/A";
-			InterimCameraData.Static.Camera.Model = "N/A";
-			InterimCameraData.Static.Camera.SerialNumber = "N/A";
-
-			StaticData = Cache->TryAddCamera(InterimCameraData);
-
-			if (ensure(StaticData.IsSet())) // Should not fail since we populated the static camera data.
-			{
-				DefaultSubjectName = Cache->GetSubjectNameFromData(InterimCameraData);
-
-				Client->PushSubjectStaticData_AnyThread({ SourceGuid, DefaultSubjectName }, ULiveLinkLensRole::StaticClass(), MoveTemp(*StaticData));
-			}
+			Client->PushSubjectStaticData_AnyThread({ SourceGuid, Cache->DefaultSubjectName }, ULiveLinkLensRole::StaticClass(), MoveTemp(*StaticData));
 		}
 	}
 
 	TArray<FSubjectAndFrameDataTuple> FrameData = Cache->TryReadFrameData(InData);
+
 	/**
 	 * This loop will consume the data in FrameData and move it into the Live Link thread. You should consider the data invalid
 	 * after this loop completes.
diff --git a/Engine/Plugins/VirtualProduction/LiveLinkOpenTrackIO/Source/LiveLinkOpenTrackIO/Private/LiveLinkOpenTrackIOSource.h b/Engine/Plugins/VirtualProduction/LiveLinkOpenTrackIO/Source/LiveLinkOpenTrackIO/Private/LiveLinkOpenTrackIOSource.h
--- a/Engine/Plugins/VirtualProduction/LiveLinkOpenTrackIO/Source/LiveLinkOpenTrackIO/Private/LiveLinkOpenTrackIOSource.h
+++ b/Engine/Plugins/VirtualProduction/LiveLinkOpenTrackIO/Source/LiveLinkOpenTrackIO/Private/LiveLinkOpenTrackIOSource.h
@@ -192,7 +192,4 @@ private:
 
 	/** Counter to track the number of inbound OpenTrackIO packets received. */
 	std::atomic<uint64> OpenTrackIOPacketCount = 0;
-	
-	/** Default subject name, used while no static camera data has arrived */
-	FName DefaultSubjectName;
 };
diff --git a/Engine/Plugins/VirtualProduction/LiveLinkOpenTrackIO/Source/LiveLinkOpenTrackIO/Private/LiveLinkOpenTrackIOTranscoder.cpp b/Engine/Plugins/VirtualProduction/LiveLinkOpenTrackIO/Source/LiveLinkOpenTrackIO/Private/LiveLinkOpenTrackIOTranscoder.cpp
--- a/Engine/Plugins/VirtualProduction/LiveLinkOpenTrackIO/Source/LiveLinkOpenTrackIO/Private/LiveLinkOpenTrackIOTranscoder.cpp
+++ b/Engine/Plugins/VirtualProduction/LiveLinkOpenTrackIO/Source/LiveLinkOpenTrackIO/Private/LiveLinkOpenTrackIOTranscoder.cpp
@@ -29,46 +29,55 @@ bool AddStaticData(FName SourceId, const T& Item, TMap<FName, T>& Container)
 
 
 FName FLiveLinkOpenTrackIOCache::GetSubjectNameFromData(const FLiveLinkOpenTrackIOData& InData)
-{
-	const FString SourceNumberAsString = FString::FromInt(InData.SourceNumber);
-	return UE::OpenTrackIO::ConvertTypeToFName(InData.Static.Camera, SourceNumberAsString); 
+{	
+	if (InData.Static.Camera.IsValid())
+	{
+		return UE::OpenTrackIO::ConvertTypeToFName(InData.Static.Camera, FString::FromInt(InData.SourceNumber));
+	}
+
+	return DefaultSubjectName;
 }
 
 TOptional<FLiveLinkStaticDataStruct> FLiveLinkOpenTrackIOCache::TryAddCamera(const FLiveLinkOpenTrackIOData& Data)
 {
-	FName SourceId = FName(Data.SourceId);
-	bool bIsNewStatic = AddStaticData(SourceId, Data.Static.Lens, StaticLenses)
-					 || AddStaticData(SourceId, Data.Static.Camera, StaticCameras);
+	const FName SourceId = FName(Data.SourceId);
+
+	const bool bIsNewStatic =
+		   AddStaticData(SourceId, Data.Static.Lens, StaticLenses)
+		|| AddStaticData(SourceId, Data.Static.Camera, StaticCameras);
 
 	if (!bIsNewStatic)
 	{
 		return {}; 
 	}
 
+	return MakeStaticData(Data);
+}
+
+FLiveLinkStaticDataStruct FLiveLinkOpenTrackIOCache::MakeStaticData(const FLiveLinkOpenTrackIOData& Data)
+{
 	/** Lens class inherits from the camera data so we use that as our base for the static struct. */
 	FLiveLinkStaticDataStruct StaticDataStruct(FLiveLinkLensStaticData::StaticStruct());
 	FLiveLinkLensStaticData& LensData = *StaticDataStruct.Cast<FLiveLinkLensStaticData>();
+
 	LensData.bIsFocalLengthSupported = true;
 	LensData.bIsApertureSupported = true;
 	LensData.bIsFocusDistanceSupported = true;
 	LensData.bIsDynamicFilmbackSupported = true;
-	
+
 	LensData.FilmBackHeight = Data.Static.Camera.ActiveSensorPhysicalDimensions.Height;
 	LensData.FilmBackWidth = Data.Static.Camera.ActiveSensorPhysicalDimensions.Width;
 
 	/** OpenTrackIO only allows for spherical models right now. */
 	LensData.LensModel = FName(TEXT("Spherical"));
-	
+
 	return StaticDataStruct;
 }
 
+
 TArray<FSubjectAndFrameDataTuple> FLiveLinkOpenTrackIOCache::TryReadFrameData(const FLiveLinkOpenTrackIOData& Data)
 {
 	TArray<FSubjectAndFrameDataTuple> Result;
-	if (Data.Transforms.IsEmpty())
-	{
-		return Result;
-	}
 	
 	/** Compute the final transform by concatentating the transforms in the frame data. */
 	FTransform FinalTransform; 
@@ -107,39 +116,37 @@ TArray<FSubjectAndFrameDataTuple> FLiveLinkOpenTrackIOCache::TryReadFrameData(co
 	double WorldTime = FPlatformTime::Seconds();
 
 	const FName SourceId = FName(Data.SourceId);
-	
-	// We require cameras to exist but it is optional to have a lens.
-	//
 	const FLiveLinkOpenTrackIOStaticCamera* Camera = StaticCameras.Find(SourceId);
-	if (Camera)
-	{
-		const FLiveLinkOpenTrackIOStaticLens* Lens = StaticLenses.Find(SourceId);
+	const FLiveLinkOpenTrackIOStaticLens* Lens = StaticLenses.Find(SourceId);
 		
-		FLiveLinkFrameDataStruct FrameDataStruct = FLiveLinkFrameDataStruct(FLiveLinkLensFrameData::StaticStruct());
-		FLiveLinkLensFrameData& CameraAndLensData = *FrameDataStruct.Cast<FLiveLinkLensFrameData>();
+	FLiveLinkFrameDataStruct FrameDataStruct = FLiveLinkFrameDataStruct(FLiveLinkLensFrameData::StaticStruct());
+	FLiveLinkLensFrameData& CameraAndLensData = *FrameDataStruct.Cast<FLiveLinkLensFrameData>();
 
-		CameraAndLensData.Transform = FinalTransform;
-		CameraAndLensData.MetaData.SceneTime = SampleTC;
-		CameraAndLensData.WorldTime = WorldTime;
+	CameraAndLensData.Transform = FinalTransform;
+	CameraAndLensData.MetaData.SceneTime = SampleTC;
+	CameraAndLensData.WorldTime = WorldTime;
 
-		CameraAndLensData.FocusDistance = Data.Lens.FocusDistance;
-		CameraAndLensData.Aperture = Data.Lens.FStop;
-		CameraAndLensData.FocalLength = Data.Lens.PinholeFocalLength;
+	CameraAndLensData.FocusDistance = Data.Lens.FocusDistance;
+	CameraAndLensData.Aperture = Data.Lens.FStop;
+	CameraAndLensData.FocalLength = Data.Lens.PinholeFocalLength;
 
+	if (Camera)
+	{
 		CameraAndLensData.FilmBackHeight = Camera->ActiveSensorPhysicalDimensions.Height;
 		CameraAndLensData.FilmBackWidth = Camera->ActiveSensorPhysicalDimensions.Width;
+	}
 		
-		if (Lens)
-		{
-			CameraAndLensData.PrincipalPoint = Data.Lens.DistortionOffset.AsVector();
-		}
-		FName SubjectName = GetSubjectNameFromData(Data);
-
-		// Note that although we are only adding frame data for a single subject. This function returns an array to
-		// support sending frame data for other parts of the OpenTrackIO spec.
-		// 
-		Result.Add({SubjectName, MoveTemp(FrameDataStruct)});
+	if (Lens)
+	{
+		CameraAndLensData.PrincipalPoint = Data.Lens.DistortionOffset.AsVector();
 	}
 
+	FName SubjectName = GetSubjectNameFromData(Data);
+
+	// Note that although we are only adding frame data for a single subject. This function returns an array to
+	// support sending frame data for other parts of the OpenTrackIO spec.
+	// 
+	Result.Add({SubjectName, MoveTemp(FrameDataStruct)});
+
 	return Result;
 }
diff --git a/Engine/Plugins/VirtualProduction/LiveLinkOpenTrackIO/Source/LiveLinkOpenTrackIO/Private/LiveLinkOpenTrackIOTranscoder.h b/Engine/Plugins/VirtualProduction/LiveLinkOpenTrackIO/Source/LiveLinkOpenTrackIO/Private/LiveLinkOpenTrackIOTranscoder.h
--- a/Engine/Plugins/VirtualProduction/LiveLinkOpenTrackIO/Source/LiveLinkOpenTrackIO/Private/LiveLinkOpenTrackIOTranscoder.h
+++ b/Engine/Plugins/VirtualProduction/LiveLinkOpenTrackIO/Source/LiveLinkOpenTrackIO/Private/LiveLinkOpenTrackIOTranscoder.h
@@ -20,6 +20,10 @@ struct FLiveLinkOpenTrackIOCache
 	TMap<FName, FLiveLinkOpenTrackIOStaticCamera> StaticCameras;
 	TMap<FName, FLiveLinkOpenTrackIOStaticLens>   StaticLenses;
 
+	/** Subject name to use if it fails to make one from the data. */
+	FName DefaultSubjectName;
+
+	/** Constructs a subject name from the given data. */
 	FName GetSubjectNameFromData(const FLiveLinkOpenTrackIOData& Data);
 	
 	/**
@@ -28,6 +32,9 @@ struct FLiveLinkOpenTrackIOCache
 	 */
 	TOptional<FLiveLinkStaticDataStruct> TryAddCamera(const FLiveLinkOpenTrackIOData& StaticData);
 
+	/** Create Live Link static data from the OpenTrackIO data */
+	FLiveLinkStaticDataStruct MakeStaticData(const FLiveLinkOpenTrackIOData& StaticData);
+
 	/**
 	 * Attempt to read the per-frame data for the given OpenTrackIO sample. If valid data is present then a constructed frame data will be provided
 	 * for each subject in the cache.
