OpenTrackIO: Added unit conversions of basic lens parameters.

Moved the OTrIO to Unreal conversions to standalone functions, to separate the data structures from the conversion logic.

Also added LiveLinkCamera plugin as a dependency, or the camera controller won't be available by default.

#rnx

diff --git a/Engine/Plugins/VirtualProduction/LiveLinkOpenTrackIO/LiveLinkOpenTrackIO.uplugin b/Engine/Plugins/VirtualProduction/LiveLinkOpenTrackIO/LiveLinkOpenTrackIO.uplugin
--- a/Engine/Plugins/VirtualProduction/LiveLinkOpenTrackIO/LiveLinkOpenTrackIO.uplugin
+++ b/Engine/Plugins/VirtualProduction/LiveLinkOpenTrackIO/LiveLinkOpenTrackIO.uplugin
@@ -26,11 +26,15 @@
 	],
 	"Plugins": [
 		{
-			"Name": "LiveLinkLens",
+			"Name": "LiveLink",
 			"Enabled": true
 		},
 		{
-			"Name": "LiveLink",
+			"Name": "LiveLinkCamera",
+			"Enabled": true
+		},
+		{
+			"Name": "LiveLinkLens",
 			"Enabled": true
 		}
 	]
diff --git a/Engine/Plugins/VirtualProduction/LiveLinkOpenTrackIO/Source/LiveLinkOpenTrackIO/Private/LiveLinkOpenTrackIOConversions.h b/Engine/Plugins/VirtualProduction/LiveLinkOpenTrackIO/Source/LiveLinkOpenTrackIO/Private/LiveLinkOpenTrackIOConversions.h
new file mode 100644
--- /dev/null
+++ b/Engine/Plugins/VirtualProduction/LiveLinkOpenTrackIO/Source/LiveLinkOpenTrackIO/Private/LiveLinkOpenTrackIOConversions.h
@@ -0,0 +1,93 @@
+// Copyright Epic Games, Inc. All Rights Reserved.
+
+#pragma once
+
+
+#include "LiveLinkLensTypes.h"
+#include "LiveLinkOpenTrackIOTypes.h"
+
+
+/** Conversion helper functions from OpenTrackIO coordinates and units to Unreal */
+namespace LiveLinkOpenTrackIOConversions
+{
+	/** OpenTrackIO coordinate system
+	 *
+	 * OpenTrackIO is:
+	 * o Z-up
+	 * o Y-forward
+	 * o Right-handed
+	 *
+	 * Unreal is:
+	 * o Z-up
+	 * o X-forward
+	 * o Left-handed
+	 *
+	 * To convert translation and scale, swap X and Y, which both updates the forward vector and toggles the handedness.
+	 * Translation also needs a factor of 100 because OpenTrackIO units are meters, while Unreal uses centimeters.
+	 * To convert rotations, invert the sign of Pan/Yaw. The others can stay the same.
+	 * Rotation units are degrees in both OpenTrackIO and Unreal.
+	 */
+
+	/** Factor to convert from OpenTrackIO's meters to Unreal's centimeters */
+	inline constexpr float MetersToCentimeters = 100.f;
+
+
+	/** Populate Live Link Lens Frame Data from relevant OpenTrackIO structures */
+	inline void ConvertLens(
+		FLiveLinkLensFrameData& OutUnrealLensData, 
+		const FLiveLinkOpenTrackIOLens* InLensData,
+		const FLiveLinkOpenTrackIOStaticCamera* InCamera
+	)
+	{
+		if (InLensData)
+		{
+			OutUnrealLensData.FocusDistance = InLensData->FocusDistance * MetersToCentimeters;
+			OutUnrealLensData.Aperture = InLensData->FStop;
+			OutUnrealLensData.FocalLength = InLensData->PinholeFocalLength;  // Both are in mm
+		}
+
+		if (InCamera)
+		{
+			OutUnrealLensData.FilmBackHeight = InCamera->ActiveSensorPhysicalDimensions.Height; // Both are in mm
+			OutUnrealLensData.FilmBackWidth = InCamera->ActiveSensorPhysicalDimensions.Width;   // Both are in mm
+
+			if (InLensData)
+			{
+				// OpenTrackIO DistortionOffset is in mm, Unreal PrincipalPoint is normalized 0..1, centered at 0.5.
+				// We scale the input offset in mm by the filmback dimensions that are also in mm.
+
+				OutUnrealLensData.PrincipalPoint.X = 0.5 + InLensData->DistortionOffset.X * OutUnrealLensData.FilmBackWidth;
+				OutUnrealLensData.PrincipalPoint.Y = 0.5 + InLensData->DistortionOffset.Y * OutUnrealLensData.FilmBackHeight;
+			}
+		}
+	}
+
+	/** From OpenTrackIO to Unreal Translation: Swap X <-> Y and convert meters to cm */
+	FVector ConvertTranslation(const FLiveLinkOpenTrackIO_XYZ& InXYZ)
+	{
+		return FVector(InXYZ.Y, InXYZ.X, InXYZ.Z) * MetersToCentimeters;
+	}
+
+	/** From OpenTrackIO to Unreal Scale: Swap X <-> Y (unitless) */
+	FVector ConvertScale(const FLiveLinkOpenTrackIO_XYZ& InXYZ)
+	{
+		return FVector(InXYZ.Y, InXYZ.X, InXYZ.Z);
+	}
+
+	/** From OpenTrackIO to Unreal Translation: Invert sign of Yaw */
+	FRotator ConvertRotation(const FLiveLinkOpenTrackIO_Rotator& InRotator)
+	{
+		return FRotator(InRotator.Tilt, -InRotator.Pan, InRotator.Roll);
+	}
+
+	/** From OpenTrackIO to Unreal FTransform */
+	FTransform ConvertTransform(const FLiveLinkOpenTrackIOTransform& InTransform)
+	{
+		return FTransform(
+			ConvertRotation(InTransform.Rotation),
+			ConvertTranslation(InTransform.Translation),
+			ConvertScale(InTransform.Scale)
+		);
+	}
+}
+
diff --git a/Engine/Plugins/VirtualProduction/LiveLinkOpenTrackIO/Source/LiveLinkOpenTrackIO/Private/LiveLinkOpenTrackIOTranscoder.cpp b/Engine/Plugins/VirtualProduction/LiveLinkOpenTrackIO/Source/LiveLinkOpenTrackIO/Private/LiveLinkOpenTrackIOTranscoder.cpp
--- a/Engine/Plugins/VirtualProduction/LiveLinkOpenTrackIO/Source/LiveLinkOpenTrackIO/Private/LiveLinkOpenTrackIOTranscoder.cpp
+++ b/Engine/Plugins/VirtualProduction/LiveLinkOpenTrackIO/Source/LiveLinkOpenTrackIO/Private/LiveLinkOpenTrackIOTranscoder.cpp
@@ -2,11 +2,11 @@
 // 
 #include "LiveLinkOpenTrackIOTranscoder.h"
 
-#include "Misc/QualifiedFrameTime.h"
-#include "Roles/LiveLinkCameraTypes.h"
 #include "LiveLinkLensTypes.h"
-#include "Roles/LiveLinkCameraTypes.h"
+#include "LiveLinkOpenTrackIOConversions.h"
 #include "Misc/App.h"
+#include "Misc/QualifiedFrameTime.h"
+#include "Roles/LiveLinkCameraTypes.h"
 
 template<typename T>
 bool AddStaticData(FName SourceId, const T& Item, TMap<FName, T>& Container)
@@ -77,13 +77,15 @@ FLiveLinkStaticDataStruct FLiveLinkOpenTrackIOCache::MakeStaticData(const FLiveL
 
 TArray<FSubjectAndFrameDataTuple> FLiveLinkOpenTrackIOCache::TryReadFrameData(const FLiveLinkOpenTrackIOData& Data)
 {
+	using namespace LiveLinkOpenTrackIOConversions;
+
 	TArray<FSubjectAndFrameDataTuple> Result;
 	
 	/** Compute the final transform by concatentating the transforms in the frame data. */
 	FTransform FinalTransform; 
 	for (const FLiveLinkOpenTrackIOTransform& Transform : Data.Transforms)
 	{
-		FinalTransform *= Transform.ToUnrealTransform();
+		FinalTransform *= ConvertTransform(Transform);
 	}
 
 	// Apply entrance pupil offset as a final transform in the forward direction.
@@ -92,7 +94,7 @@ TArray<FSubjectAndFrameDataTuple> FLiveLinkOpenTrackIOCache::TryReadFrameData(co
 
 		const FTransform PupilOffset(
 			FQuat::Identity,
-			EntrancePupilAsXYZ.ToUnrealTranslation(),
+			ConvertTranslation(EntrancePupilAsXYZ),
 			FVector(1, 1, 1)
 		);
 
@@ -113,8 +115,6 @@ TArray<FSubjectAndFrameDataTuple> FLiveLinkOpenTrackIOCache::TryReadFrameData(co
 		SampleTC = Data.Timing.Timecode.GetQualifiedFrameTime();
 	}
 
-	double WorldTime = FPlatformTime::Seconds();
-
 	const FName SourceId = FName(Data.SourceId);
 	const FLiveLinkOpenTrackIOStaticCamera* Camera = StaticCameras.Find(SourceId);
 	const FLiveLinkOpenTrackIOStaticLens* Lens = StaticLenses.Find(SourceId);
@@ -122,24 +122,11 @@ TArray<FSubjectAndFrameDataTuple> FLiveLinkOpenTrackIOCache::TryReadFrameData(co
 	FLiveLinkFrameDataStruct FrameDataStruct = FLiveLinkFrameDataStruct(FLiveLinkLensFrameData::StaticStruct());
 	FLiveLinkLensFrameData& CameraAndLensData = *FrameDataStruct.Cast<FLiveLinkLensFrameData>();
 
+	ConvertLens(CameraAndLensData, &Data.Lens, Camera);
+
 	CameraAndLensData.Transform = FinalTransform;
 	CameraAndLensData.MetaData.SceneTime = SampleTC;
-	CameraAndLensData.WorldTime = WorldTime;
-
-	CameraAndLensData.FocusDistance = Data.Lens.FocusDistance;
-	CameraAndLensData.Aperture = Data.Lens.FStop;
-	CameraAndLensData.FocalLength = Data.Lens.PinholeFocalLength;
-
-	if (Camera)
-	{
-		CameraAndLensData.FilmBackHeight = Camera->ActiveSensorPhysicalDimensions.Height;
-		CameraAndLensData.FilmBackWidth = Camera->ActiveSensorPhysicalDimensions.Width;
-	}
-		
-	if (Lens)
-	{
-		CameraAndLensData.PrincipalPoint = Data.Lens.DistortionOffset.AsVector();
-	}
+	CameraAndLensData.WorldTime = FPlatformTime::Seconds();
 
 	FName SubjectName = GetSubjectNameFromData(Data);
 
