OpenTrackIO: Fixed CBOR parsing.

* Upgraded scalars to doubles because FCborStructDeserializerBackend doesn't deserialize doubles into floats, and producers may send doubles as the opentrackio.org doesn't restrict the types.
* Initialized cbor backed with ECborEndianness::StandardCompliant because it defaults to platform endiannes instead.
* Initialized cbor backed with bInIsLWCCompatibilityMode so that it allows deserializing floats into doubles.

#rnx

diff --git a/Engine/Plugins/VirtualProduction/LiveLinkOpenTrackIO/Source/LiveLinkOpenTrackIO/Private/LiveLinkOpenTrackIOParser.cpp b/Engine/Plugins/VirtualProduction/LiveLinkOpenTrackIO/Source/LiveLinkOpenTrackIO/Private/LiveLinkOpenTrackIOParser.cpp
--- a/Engine/Plugins/VirtualProduction/LiveLinkOpenTrackIO/Source/LiveLinkOpenTrackIO/Private/LiveLinkOpenTrackIOParser.cpp
+++ b/Engine/Plugins/VirtualProduction/LiveLinkOpenTrackIO/Source/LiveLinkOpenTrackIO/Private/LiveLinkOpenTrackIOParser.cpp
@@ -64,7 +64,15 @@ TOptional<FLiveLinkOpenTrackIOData> ParseCborBlob(TArrayView<uint8> InBytes)
 	}
 	
 	FMemoryReaderView Reader(InBytes);
-	FCborStructDeserializerBackend Deserializer(Reader);
+
+	// Notes about the Cbor struct deserializer backend:
+	// 
+	// * Does not support deserializing doubles into floats, hence it is safer to use structs with doubles instead.
+	// * Does not support deserializing floating point to integer types (the data will be ignored).
+	// * Set bInIsLWCCompatibilityMode to true to allow deserializing floats into doubles.
+	// * See FCborStructDeserializerBackend::ReadProperty for more details.
+	FCborStructDeserializerBackend Deserializer(Reader, ECborEndianness::StandardCompliant, true /* bInIsLWCCompatibilityMode */);
+
 	UScriptStruct* InType = TBaseStructure<FLiveLinkOpenTrackIOData>::Get();
 
 	FLiveLinkOpenTrackIOData Data;
diff --git a/Engine/Plugins/VirtualProduction/LiveLinkOpenTrackIO/Source/LiveLinkOpenTrackIO/Private/LiveLinkOpenTrackIOTypes.h b/Engine/Plugins/VirtualProduction/LiveLinkOpenTrackIO/Source/LiveLinkOpenTrackIO/Private/LiveLinkOpenTrackIOTypes.h
--- a/Engine/Plugins/VirtualProduction/LiveLinkOpenTrackIO/Source/LiveLinkOpenTrackIO/Private/LiveLinkOpenTrackIOTypes.h
+++ b/Engine/Plugins/VirtualProduction/LiveLinkOpenTrackIO/Source/LiveLinkOpenTrackIO/Private/LiveLinkOpenTrackIOTypes.h
@@ -13,30 +13,6 @@
 #include "LiveLinkOpenTrackIOTypes.generated.h"
 
 
-namespace LiveLinkOpenTrackIO
-{
-	/** OpenTrackIO coordinate system 
-	 * 
-	 * OpenTrackIO is:
-	 * o Z-up
-	 * o Y-forward
-	 * o Right-handed
-	 *
-	 * Unreal is:
-	 * o Z-up
-	 * o X-forward
-	 * o Left-handed
-	 *
-	 * To convert translation and scale, swap X and Y, which both updates the forward vector and toggles the handedness.
-	 * Translation also needs a factor of 100 because OpenTrackIO units are meters, while Unreal uses centimeters.
-	 * To convert rotations, invert the sign of Pan/Yaw. The others can stay the same.
-	 * Rotation units are degrees in both OpenTrackIO and Unreal.
-	 */
-
-	/** Factor to convert from OpenTrackIO's meters to Unreal's centimeters */
-	inline constexpr float MetersToCentimeters = 100.f;
-}
-
 /**
  * Represents frame rate type from OpenTrack I/O. Value is represented in hertz using an fractional representation
  * Numerator / Denominator. The default value is 24 / 1 or 24 cycles per second.
@@ -129,10 +105,10 @@ struct FLiveLinkOpenTrackIOSensorDimensions
 	bool operator!=(const FLiveLinkOpenTrackIOSensorDimensions&) const = default;
 	
 	UPROPERTY(EditAnywhere, BlueprintReadWrite, Category = "OpenTrackIO")
-	float Height = 43.0f;
+	double Height = 43.0f;
 
 	UPROPERTY(EditAnywhere, BlueprintReadWrite, Category = "OpenTrackIO")
-	float Width = 24.0f;
+	double Width = 24.0f;
 };
 
 /**
@@ -278,7 +254,7 @@ struct FLiveLinkOpenTrackIOStaticCamera
 	 * to the value of the parameter divided by 360 times the capture frame rate.
 	 */
 	UPROPERTY(meta = (ClampMin = "0.0", ClampMax = "360.0"))
-	float ShutterAngle = 180.0;
+	double ShutterAngle = 180.0;
 };
 
 USTRUCT(BlueprintType)
@@ -313,7 +289,7 @@ struct FLiveLinkOpenTrackIOStaticLens
 	 * of both projection-matrix-based and field-of-view-based rendering as per the OpenLensIO documentation.
 	 */
 	UPROPERTY(EditAnywhere, BlueprintReadWrite, Category = "OpenTrackIO")
-	float DistortionOverscanMax = 1.0;
+	double DistortionOverscanMax = 1.0;
 
 	/**
 	 * Static maximum overscan factor on lens undistortion. This is an alternative
@@ -322,7 +298,7 @@ struct FLiveLinkOpenTrackIOStaticLens
 	 * documentation.
 	 */
 	UPROPERTY(EditAnywhere, BlueprintReadWrite, Category = "OpenTrackIO")
-	float UndistortionOverscanMax = 1.0;
+	double UndistortionOverscanMax = 1.0;
 
 	/** Non-blank string naming lens manufacturer */
 	UPROPERTY(EditAnywhere, BlueprintReadWrite, Category = "OpenTrackIO")
@@ -342,7 +318,7 @@ struct FLiveLinkOpenTrackIOStaticLens
 
 	/** Nominal focal length of the lens. The number printed on the side of a prime lens, e.g. 50 mm, and undefined in the case of a zoom lens. */
 	UPROPERTY(EditAnywhere, BlueprintReadWrite, Category = "OpenTrackIO")
-	float NominalFocalLength = 50;
+	double NominalFocalLength = 50;
 
 	/** List of free strings that describe the history of calibrations of the lens. */
 	UPROPERTY(EditAnywhere, BlueprintReadWrite, Category = "OpenTrackIO")
@@ -519,13 +495,13 @@ struct FLiveLinkOpenTrackIOSynchronizationOffsets
 	bool operator!=(const FLiveLinkOpenTrackIOSynchronizationOffsets&) const = default;
 	
 	UPROPERTY(EditAnywhere, BlueprintReadWrite, Category = "OpenTrackIO")
-	float Translation = 0.0f;
+	double Translation = 0.0f;
 
 	UPROPERTY(EditAnywhere, BlueprintReadWrite, Category = "OpenTrackIO")
-	float Rotation = 0.0f;
+	double Rotation = 0.0f;
 
 	UPROPERTY(EditAnywhere, BlueprintReadWrite, Category = "OpenTrackIO")
-	float LensEncoders = 0.0f;
+	double LensEncoders = 0.0f;
 };
 
 USTRUCT(BlueprintType)
@@ -585,14 +561,14 @@ struct FLiveLinkOpenTrackIO_PTP
 	FLiveLinkOpenTrackIO_PTPLeaderPriorities LeaderPriorities;
 
 	UPROPERTY(EditAnywhere, BlueprintReadWrite, Category = "OpenTrackIO")
-	float LeaderAccuracy = 0.0;
+	double LeaderAccuracy = 0.0;
 
 	/** Enum string property [GNSS, Atomic clock, NTP] */
 	UPROPERTY(EditAnywhere, BlueprintReadWrite, Category = "OpenTrackIO")
 	FString LeaderTimeSource = TEXT("NTP");
 
 	UPROPERTY(EditAnywhere, BlueprintReadWrite, Category = "OpenTrackIO")
-	float MeanPathDelay = 0.0;
+	double MeanPathDelay = 0.0;
 
 	UPROPERTY(EditAnywhere, BlueprintReadWrite, Category = "OpenTrackIO")
 	int32 Vlan = 0;
@@ -821,10 +797,10 @@ struct FLiveLinkOpenTrackIOLens_DistortionOffset
 	}
 	
 	UPROPERTY(EditAnywhere, BlueprintReadWrite, Category = "OpenTrackIO")
-	float X = 0;
+	double X = 0;
 
 	UPROPERTY(EditAnywhere, BlueprintReadWrite, Category = "OpenTrackIO")
-	float Y = 0;
+	double Y = 0;
 };
 
 /**
@@ -855,13 +831,13 @@ struct FLiveLinkOpenTrackIOLens_Encoders
 	bool operator!=(const FLiveLinkOpenTrackIOLens_Encoders&) const = default;
 	
 	UPROPERTY(EditAnywhere, BlueprintReadWrite, Category = "OpenTrackIO")
-	float Focus = 0;
+	double Focus = 0;
 
 	UPROPERTY(EditAnywhere, BlueprintReadWrite, Category = "OpenTrackIO")
-	float Iris = 0;
+	double Iris = 0;
 
 	UPROPERTY(EditAnywhere, BlueprintReadWrite, Category = "OpenTrackIO")
-	float Zoom = 0;
+	double Zoom = 0;
 };
 
 
@@ -914,13 +890,13 @@ struct FLiveLinkOpenTrackIOLens_ExposureFalloff
 	bool operator!=(const FLiveLinkOpenTrackIOLens_ExposureFalloff&) const = default;
 	
 	UPROPERTY(EditAnywhere, BlueprintReadWrite, Category = "OpenTrackIO")
-	float A1 = 0;
+	double A1 = 0;
 
 	UPROPERTY(EditAnywhere, BlueprintReadWrite, Category = "OpenTrackIO")
-	float A2 = 0;
+	double A2 = 0;
 
 	UPROPERTY(EditAnywhere, BlueprintReadWrite, Category = "OpenTrackIO")
-	float A3 = 0;
+	double A3 = 0;
 };
 
 USTRUCT(BlueprintType)
@@ -939,13 +915,13 @@ struct FLiveLinkOpenTrackIOLens_DistortionCoeff
 	bool operator!=(const FLiveLinkOpenTrackIOLens_DistortionCoeff&) const = default;
 	
 	UPROPERTY(EditAnywhere, BlueprintReadWrite, Category = "OpenTrackIO")
-	TArray<float> Radial;
+	TArray<double> Radial;
 
 	UPROPERTY(EditAnywhere, BlueprintReadWrite, Category = "OpenTrackIO")
-	TArray<float> Tangential;
+	TArray<double> Tangential;
 
 	UPROPERTY(EditAnywhere, BlueprintReadWrite, Category = "OpenTrackIO")
-	float Overscan = 0.0;
+	double Overscan = 0.0;
 };
 
 USTRUCT(BlueprintType)
@@ -969,7 +945,7 @@ struct FLiveLinkOpenTrackIOLens
 	 * negotiation between a particular producer and consumer.
 	 */
 	UPROPERTY(EditAnywhere, BlueprintReadWrite, Category = "OpenTrackIO")
-	TArray<float> Custom;
+	TArray<double> Custom;
 
 	/**
 	 * A list of Distortion objects that each define the coefficients for calculating the distortion characteristics of a lens comprising radial
@@ -998,7 +974,7 @@ struct FLiveLinkOpenTrackIOLens
 	 * virtual camera’s transform chain.
 	 */
 	UPROPERTY(EditAnywhere, BlueprintReadWrite, Category = "OpenTrackIO")
-	float  EntrancePupilOffset = 0;
+	double  EntrancePupilOffset = 0;
 
 	/**
 	 * Coefficients for calculating the exposure fall-off (vignetting) of a lens
@@ -1008,15 +984,15 @@ struct FLiveLinkOpenTrackIOLens
 
 	/** The linear f-number of the lens, equal to the focal length divided by the diameter of the entrance pupil. */
 	UPROPERTY(EditAnywhere, BlueprintReadWrite, Category = "OpenTrackIO")
-	float FStop = 2.8f;
+	double FStop = 2.8f;
 
 	/** Distance between the pinhole and the image plane in the simple CGI pinhole camera model. (millimeters) */
 	UPROPERTY(EditAnywhere, BlueprintReadWrite, Category = "OpenTrackIO")
-	float PinholeFocalLength = 35.0f;
+	double PinholeFocalLength = 35.0f;
 
 	/** Focus distance/position of the lens in meters. */ 
 	UPROPERTY(EditAnywhere, BlueprintReadWrite, Category = "OpenTrackIO")
-	float FocusDistance = 10.0f;
+	double FocusDistance = 10.0f;
 
 	/**  Offset in x and y of the centre of perspective projection of the virtual camera */
 	UPROPERTY(EditAnywhere, BlueprintReadWrite, Category = "OpenTrackIO")
@@ -1028,7 +1004,7 @@ struct FLiveLinkOpenTrackIOLens
 
 	/** T-stop: Linear t-number of the lens, equal to the F-number of the lens divided by the square root of the transmittance of the lens. */
 	UPROPERTY(EditAnywhere, BlueprintReadWrite, Category = "OpenTrackIO")
-	float TStop = 0.0;
+	double TStop = 0.0;
 };
 
 /**
@@ -1113,22 +1089,22 @@ struct FLiveLinkOpenTrackIOGlobalStage
 	bool operator!=(const FLiveLinkOpenTrackIOGlobalStage&) const = default;
 
 	UPROPERTY(EditAnywhere, BlueprintReadWrite, Category = "OpenTrackIO")
-	float E = 0;
+	double E = 0;
 
 	UPROPERTY(EditAnywhere, BlueprintReadWrite, Category = "OpenTrackIO")
-	float N = 0;
+	double N = 0;
 
 	UPROPERTY(EditAnywhere, BlueprintReadWrite, Category = "OpenTrackIO")
-	float U = 0;
+	double U = 0;
 
 	UPROPERTY(EditAnywhere, BlueprintReadWrite, Category = "OpenTrackIO")
-	float Lat0 = 0;
+	double Lat0 = 0;
 
 	UPROPERTY(EditAnywhere, BlueprintReadWrite, Category = "OpenTrackIO")
-	float Lon0 = 0;
+	double Lon0 = 0;
 
 	UPROPERTY(EditAnywhere, BlueprintReadWrite, Category = "OpenTrackIO")
-	float H0 = 0;
+	double H0 = 0;
 };
 
 USTRUCT(BlueprintType)
@@ -1146,20 +1122,14 @@ struct FLiveLinkOpenTrackIO_Rotator
 	bool operator==(const FLiveLinkOpenTrackIO_Rotator&) const = default;
 	bool operator!=(const FLiveLinkOpenTrackIO_Rotator&) const = default;
 	
-	/** From OpenTrackIO to Unreal Translation: Invert sign of Yaw */
-	FRotator ToUnrealRotator() const
-	{
-		return FRotator(Tilt, -Pan, Roll);
-	}
-
 	UPROPERTY(EditAnywhere, BlueprintReadWrite, Category = "OpenTrackIO")
-	float Pan = 0;
+	double Pan = 0;
 
 	UPROPERTY(EditAnywhere, BlueprintReadWrite, Category = "OpenTrackIO")
-	float Tilt = 0;
+	double Tilt = 0;
 
 	UPROPERTY(EditAnywhere, BlueprintReadWrite, Category = "OpenTrackIO")
-	float Roll = 0;
+	double Roll = 0;
 };
 
 USTRUCT(BlueprintType)
@@ -1168,7 +1138,7 @@ struct FLiveLinkOpenTrackIO_XYZ
 	GENERATED_BODY()
 
 	FLiveLinkOpenTrackIO_XYZ() = default;
-	FLiveLinkOpenTrackIO_XYZ(float InX, float InY, float InZ)
+	FLiveLinkOpenTrackIO_XYZ(double InX, double InY, double InZ)
 		: X(InX), Y(InY), Z(InZ)
 	{
 	}
@@ -1183,25 +1153,14 @@ struct FLiveLinkOpenTrackIO_XYZ
 	bool operator==(const FLiveLinkOpenTrackIO_XYZ&) const = default;
 	bool operator!=(const FLiveLinkOpenTrackIO_XYZ&) const = default;
 
-	/** From OpenTrackIO to Unreal Translation: Swap X <-> Y and convert meters to cm */
-	FVector ToUnrealTranslation() const
-	{
-		return FVector(Y, X, Z) * LiveLinkOpenTrackIO::MetersToCentimeters;
-	}
-
-	/** From OpenTrackIO to Unreal Scale: Swap X <-> Y (unitless) */
-	FVector ToUnrealScale() const
-	{
-		return FVector(Y, X, Z);
-	}
-
 	UPROPERTY(EditAnywhere, BlueprintReadWrite, Category = "OpenTrackIO")
-	float X = 0;
+	double X = 0;
+
 	UPROPERTY(EditAnywhere, BlueprintReadWrite, Category = "OpenTrackIO")
-	float Y = 0;
+	double Y = 0;
 
 	UPROPERTY(EditAnywhere, BlueprintReadWrite, Category = "OpenTrackIO")
-	float Z = 0;
+	double Z = 0;
 };
 
 USTRUCT(BlueprintType)
@@ -1218,16 +1177,6 @@ struct FLiveLinkOpenTrackIOTransform
 	
 	bool operator==(const FLiveLinkOpenTrackIOTransform&) const = default;
 	bool operator!=(const FLiveLinkOpenTrackIOTransform&) const = default;
-
-	/** From OpenTrackIO to Unreal FTransform */
-	FTransform ToUnrealTransform() const
-	{
-		return FTransform(
-			Rotation.ToUnrealRotator(),
-			Translation.ToUnrealTranslation(),
-			Scale.ToUnrealScale()
-		);
-	}
 	
 	/** Translation part of a OpenTrackIO xform */
 	UPROPERTY(EditAnywhere, BlueprintReadWrite, Category = "OpenTrackIO")
