Copy custom meta data from the OpenTrackIO to the LiveLink equivalent.

Append the source number to the name of the subject.  This will allow us to have unique subjects from multiple sources on the same LL source.

#rnx

diff --git a/Engine/Plugins/VirtualProduction/LiveLinkOpenTrackIO/Source/LiveLinkOpenTrackIO/Private/LiveLinkOpenTrackIOTranscoder.cpp b/Engine/Plugins/VirtualProduction/LiveLinkOpenTrackIO/Source/LiveLinkOpenTrackIO/Private/LiveLinkOpenTrackIOTranscoder.cpp
--- a/Engine/Plugins/VirtualProduction/LiveLinkOpenTrackIO/Source/LiveLinkOpenTrackIO/Private/LiveLinkOpenTrackIOTranscoder.cpp
+++ b/Engine/Plugins/VirtualProduction/LiveLinkOpenTrackIO/Source/LiveLinkOpenTrackIO/Private/LiveLinkOpenTrackIOTranscoder.cpp
@@ -128,8 +128,13 @@ TArray<FSubjectAndFrameDataTuple> FLiveLinkOpenTrackIOCache::TryReadFrameData(co
 	CameraAndLensData.MetaData.SceneTime = SampleTC;
 	CameraAndLensData.WorldTime = FPlatformTime::Seconds();
 
+	// Add any custom meta data.
+	for (const FLiveLinkOpenTrackIOCustomDataField& Field : Data.Custom.LiveLinkMetaData)
+	{
+		CameraAndLensData.MetaData.StringMetaData.Add(FName(Field.Key), Field.Value);
+	}
+	
 	FName SubjectName = GetSubjectNameFromData(Data);
-
 	// Note that although we are only adding frame data for a single subject. This function returns an array to
 	// support sending frame data for other parts of the OpenTrackIO spec.
 	// 
diff --git a/Engine/Plugins/VirtualProduction/LiveLinkOpenTrackIO/Source/LiveLinkOpenTrackIO/Private/LiveLinkOpenTrackIOTypes.h b/Engine/Plugins/VirtualProduction/LiveLinkOpenTrackIO/Source/LiveLinkOpenTrackIO/Private/LiveLinkOpenTrackIOTypes.h
--- a/Engine/Plugins/VirtualProduction/LiveLinkOpenTrackIO/Source/LiveLinkOpenTrackIO/Private/LiveLinkOpenTrackIOTypes.h
+++ b/Engine/Plugins/VirtualProduction/LiveLinkOpenTrackIO/Source/LiveLinkOpenTrackIO/Private/LiveLinkOpenTrackIOTypes.h
@@ -69,20 +69,22 @@ namespace UE::OpenTrackIO
 template<CHasModelAndMake T>
 FName ConvertTypeToFName(const T& Type, const FString& Postfix)
 {
+	FNameBuilder Builder;	
 	if constexpr(CHasLabelMember<T>)
 	{
-		return FName(*Type.Label);
+		Builder.Append(*Type.Label);
 	}
 	else
 	{
-		FNameBuilder Builder;
 		Builder.Appendf(TEXT("%s_%s"), *Type.Make, *Type.Model);
-		if (!Postfix.IsEmpty())
-		{
-			Builder.Appendf(TEXT(":%s"), *Postfix);
-		}
-		return *Builder;
 	}
+	
+	if (!Postfix.IsEmpty())
+	{
+		Builder.Appendf(TEXT(":%s"), *Postfix);
+	}
+	
+	return *Builder;
 }
 
 }
