Add sequence number checking to the inbound packets from a given source.  This ensures that packets are received in order and out-of-order packets are discarded.

#rnx

diff --git a/Engine/Plugins/VirtualProduction/LiveLinkOpenTrackIO/Source/LiveLinkOpenTrackIO/Private/LiveLinkOpenTrackIOSource.cpp b/Engine/Plugins/VirtualProduction/LiveLinkOpenTrackIO/Source/LiveLinkOpenTrackIO/Private/LiveLinkOpenTrackIOSource.cpp
--- a/Engine/Plugins/VirtualProduction/LiveLinkOpenTrackIO/Source/LiveLinkOpenTrackIO/Private/LiveLinkOpenTrackIOSource.cpp
+++ b/Engine/Plugins/VirtualProduction/LiveLinkOpenTrackIO/Source/LiveLinkOpenTrackIO/Private/LiveLinkOpenTrackIOSource.cpp
@@ -313,6 +313,7 @@ bool FLiveLinkOpenTrackIOSource::ConsumeInboundData()
 			bProcessedData = true;
 			OpenTrackIOPacketCount++;
 
+			
 			PushDataToLiveLink_AnyThread(ParsedData->Get<0>() /** Header */, ParsedData->Get<1>() /** Body */);
 		}
 		else
@@ -335,6 +336,12 @@ void FLiveLinkOpenTrackIOSource::PushDataToLiveLink_AnyThread(const FLiveLinkOpe
 		Cache = OpenTrackIOCacheMap.Add(InData.SourceId, MakePimpl<FLiveLinkOpenTrackIOCache>()).Get();
 	}
 
+	if (!Cache->IsPacketInSequence(Header.SequenceNumber, InData.Timing.SampleRate.GetFrameRate()))
+	{
+		UE_LOG(LogLiveLinkOpenTrackIO, Warning, TEXT("Received packet for %s is out of sequence. Discarding %d."), *InData.SourceId, Header.SequenceNumber);
+		return;
+	}
+	
 	/** Check the static data camera fields on the inbound struct and see if we have new data to store / update. */
 	if (TOptional<FLiveLinkStaticDataStruct> StaticData = Cache->TryAddCamera(InData))
 	{
@@ -376,6 +383,8 @@ void FLiveLinkOpenTrackIOSource::PushDataToLiveLink_AnyThread(const FLiveLinkOpe
 	{
 		Client->PushSubjectFrameData_AnyThread({ SourceGuid, Item.SubjectName }, MoveTemp(Item.FrameData));
 	}
+
+	Cache->UpdateLastKnownSequenceNumber(Header.SequenceNumber);
 }
 
 void FLiveLinkOpenTrackIOSource::CloseSockets()
diff --git a/Engine/Plugins/VirtualProduction/LiveLinkOpenTrackIO/Source/LiveLinkOpenTrackIO/Private/LiveLinkOpenTrackIOTranscoder.cpp b/Engine/Plugins/VirtualProduction/LiveLinkOpenTrackIO/Source/LiveLinkOpenTrackIO/Private/LiveLinkOpenTrackIOTranscoder.cpp
--- a/Engine/Plugins/VirtualProduction/LiveLinkOpenTrackIO/Source/LiveLinkOpenTrackIO/Private/LiveLinkOpenTrackIOTranscoder.cpp
+++ b/Engine/Plugins/VirtualProduction/LiveLinkOpenTrackIO/Source/LiveLinkOpenTrackIO/Private/LiveLinkOpenTrackIOTranscoder.cpp
@@ -27,6 +27,24 @@ bool AddStaticData(FName SourceId, const T& Item, TMap<FName, T>& Container)
 	return false;
 }
 
+bool FLiveLinkOpenTrackIOCache::IsPacketInSequence(uint16 SequenceNumber, const FFrameRate& InRate) const
+{
+	// Initially the expected sequence number is -1, which indicates we have not received our first packet yet. 
+	if (ExpectedSequenceNumber > 0 && SequenceNumber < ExpectedSequenceNumber)
+	{
+		double RateAsDecimal = InRate.AsDecimal();
+		check(RateAsDecimal>0);
+		const double SampleRateExpiry = 3 * 1.0 / InRate.AsDecimal();
+		// We we have exceeded a timeout period then consider the sequence stream as being reset.  In other words,
+		// sequence numbers are sequential only 
+		if ((FPlatformTime::Seconds() - LastDataReceiveTimeInSeconds) < SampleRateExpiry)
+		{
+			return false;
+		}
+	}
+
+	return true;
+}
 
 FName FLiveLinkOpenTrackIOCache::GetSubjectNameFromData(const FLiveLinkOpenTrackIOData& InData)
 {	
diff --git a/Engine/Plugins/VirtualProduction/LiveLinkOpenTrackIO/Source/LiveLinkOpenTrackIO/Private/LiveLinkOpenTrackIOTranscoder.h b/Engine/Plugins/VirtualProduction/LiveLinkOpenTrackIO/Source/LiveLinkOpenTrackIO/Private/LiveLinkOpenTrackIOTranscoder.h
--- a/Engine/Plugins/VirtualProduction/LiveLinkOpenTrackIO/Source/LiveLinkOpenTrackIO/Private/LiveLinkOpenTrackIOTranscoder.h
+++ b/Engine/Plugins/VirtualProduction/LiveLinkOpenTrackIO/Source/LiveLinkOpenTrackIO/Private/LiveLinkOpenTrackIOTranscoder.h
@@ -20,6 +20,29 @@ struct FLiveLinkOpenTrackIOCache
 	TMap<FName, FLiveLinkOpenTrackIOStaticCamera> StaticCameras;
 	TMap<FName, FLiveLinkOpenTrackIOStaticLens>   StaticLenses;
 
+	/**
+	 * Check the given sequence number and return true if it is the expected number. Sequence numbers should be incrementing
+	 * modulo UINT16_MAX.  This also accounts for delays in transmission based on sample rate provided by the first inbound packets.
+	 */
+	bool IsPacketInSequence(uint16 SequenceNumber, const FFrameRate& InRate) const;
+
+	/**
+	 * Caches the processed sequence number so that it can be compared when the next packet is received.
+	 */
+	void UpdateLastKnownSequenceNumber(uint16 SequenceNumber)
+	{
+		ExpectedSequenceNumber = SequenceNumber + 1;
+		if (ExpectedSequenceNumber > UINT16_MAX)
+		{
+			ExpectedSequenceNumber = 1;
+		}
+		LastDataReceiveTimeInSeconds = FPlatformTime::Seconds();
+	}
+	
+	/**
+	 * Set sequence number 
+	 */
+
 	/** Subject name to use if it fails to make one from the data. */
 	FName DefaultSubjectName;
 
@@ -40,5 +63,11 @@ struct FLiveLinkOpenTrackIOCache
 	 * for each subject in the cache.
 	 */
 	TArray<FSubjectAndFrameDataTuple> TryReadFrameData(const FLiveLinkOpenTrackIOData& Data);
-	
+
+private:
+	/** Tracks the current sequence number for the source. */
+	int32 ExpectedSequenceNumber = -1;
+
+	/** Store the last time in FPlatformTime::Seconds() we received data from this source. */
+	double LastDataReceiveTimeInSeconds = 0;
 };
