OpenTrackIO: Relaxed cbor to unreal numerical conversion.

This should make it more resilient to type mismatches such as float vs double.

#rnx

diff --git a/Engine/Plugins/VirtualProduction/LiveLinkOpenTrackIO/Source/LiveLinkOpenTrackIO/LiveLinkOpenTrackIO.Build.cs b/Engine/Plugins/VirtualProduction/LiveLinkOpenTrackIO/Source/LiveLinkOpenTrackIO/LiveLinkOpenTrackIO.Build.cs
--- a/Engine/Plugins/VirtualProduction/LiveLinkOpenTrackIO/Source/LiveLinkOpenTrackIO/LiveLinkOpenTrackIO.Build.cs
+++ b/Engine/Plugins/VirtualProduction/LiveLinkOpenTrackIO/Source/LiveLinkOpenTrackIO/LiveLinkOpenTrackIO.Build.cs
@@ -22,6 +22,7 @@ public class LiveLinkOpenTrackIO : ModuleRules
 			new string[]
 			{
 				"ApplicationCore",
+				"Cbor",
 				"CoreUObject",
 				"Engine",
 				"Json",
diff --git a/Engine/Plugins/VirtualProduction/LiveLinkOpenTrackIO/Source/LiveLinkOpenTrackIO/Private/LiveLinkOpenTrackIOParser.cpp b/Engine/Plugins/VirtualProduction/LiveLinkOpenTrackIO/Source/LiveLinkOpenTrackIO/Private/LiveLinkOpenTrackIOParser.cpp
--- a/Engine/Plugins/VirtualProduction/LiveLinkOpenTrackIO/Source/LiveLinkOpenTrackIO/Private/LiveLinkOpenTrackIOParser.cpp
+++ b/Engine/Plugins/VirtualProduction/LiveLinkOpenTrackIO/Source/LiveLinkOpenTrackIO/Private/LiveLinkOpenTrackIOParser.cpp
@@ -6,6 +6,7 @@
 
 #include "LiveLinkOpenTrackIODatagram.h"
 #include "LiveLinkOpenTrackIOTypes.h"
+#include "OpenTrackIOCborStructDeserializerBackend.h"
 
 #include "JsonObjectConverter.h"
 
@@ -19,8 +20,6 @@
 
 #include "StructSerializer.h"
 #include "StructDeserializer.h"
-#include "Backends/CborStructSerializerBackend.h"
-#include "Backends/CborStructDeserializerBackend.h"
 
 namespace UE::OpenTrackIO::Private
 {
@@ -64,14 +63,7 @@ TOptional<FLiveLinkOpenTrackIOData> ParseCborBlob(TArrayView<uint8> InBytes)
 	}
 	
 	FMemoryReaderView Reader(InBytes);
-
-	// Notes about the Cbor struct deserializer backend:
-	// 
-	// * Does not support deserializing doubles into floats, hence it is safer to use structs with doubles instead.
-	// * Does not support deserializing floating point to integer types (the data will be ignored).
-	// * Set bInIsLWCCompatibilityMode to true to allow deserializing floats into doubles.
-	// * See FCborStructDeserializerBackend::ReadProperty for more details.
-	FCborStructDeserializerBackend Deserializer(Reader, ECborEndianness::StandardCompliant, true /* bInIsLWCCompatibilityMode */);
+	FOpenTrackIOCborStructDeserializerBackend Deserializer(Reader);
 
 	UScriptStruct* InType = TBaseStructure<FLiveLinkOpenTrackIOData>::Get();
 
diff --git a/Engine/Plugins/VirtualProduction/LiveLinkOpenTrackIO/Source/LiveLinkOpenTrackIO/Private/OpenTrackIOCborStructDeserializerBackend.cpp b/Engine/Plugins/VirtualProduction/LiveLinkOpenTrackIO/Source/LiveLinkOpenTrackIO/Private/OpenTrackIOCborStructDeserializerBackend.cpp
new file mode 100644
--- /dev/null
+++ b/Engine/Plugins/VirtualProduction/LiveLinkOpenTrackIO/Source/LiveLinkOpenTrackIO/Private/OpenTrackIOCborStructDeserializerBackend.cpp
@@ -0,0 +1,516 @@
+// Copyright Epic Games, Inc. All Rights Reserved.
+
+#include "OpenTrackIOCborStructDeserializerBackend.h"
+
+#include "LiveLinkOpenTrackIO.h"
+
+#include "UObject/Class.h"
+#include "UObject/AnsiStrProperty.h"
+#include "UObject/EnumProperty.h"
+#include "UObject/TextProperty.h"
+#include "UObject/UnrealType.h"
+#include "UObject/Utf8StrProperty.h"
+
+
+namespace UE::OpenTrackIOCbor::Private
+{
+	/**
+	* Clears the value of the given property.
+	*
+	* @param Property The property to clear.
+	* @param Outer The property that contains the property to be cleared, if any.
+	* @param Data A pointer to the memory holding the property's data.
+	* @param ArrayIndex The index of the element to clear (if the property is an array).
+	* 
+	* @return true on success, false otherwise.
+	*/
+	static bool ClearPropertyValue(FProperty* Property, FProperty* Outer, void* Data, int32 ArrayIndex)
+	{
+		FArrayProperty* ArrayProperty = CastField<FArrayProperty>(Outer);
+
+		if (ArrayProperty != nullptr)
+		{
+			if (ArrayProperty->Inner != Property)
+			{
+				return false;
+			}
+
+			FScriptArrayHelper ArrayHelper(ArrayProperty, ArrayProperty->template ContainerPtrToValuePtr<void>(Data));
+			ArrayIndex = ArrayHelper.AddValue();
+		}
+
+		Property->ClearValue_InContainer(Data, ArrayIndex);
+
+		return true;
+	}
+
+
+	/**
+	* Gets a pointer to object of the given property.
+	*
+	* @param Property The property to get.
+	* @param Outer The property that contains the property to be get, if any.
+	* @param Data A pointer to the memory holding the property's data.
+	* @param ArrayIndex The index of the element to set (if the property is an array).
+	* 
+	* @return A pointer to the object represented by the property, null otherwise.
+	*/
+	static void* GetPropertyValuePtr(FProperty* Property, FProperty* Outer, void* Data, int32 ArrayIndex)
+	{
+		check(Property);
+
+		if (FArrayProperty* ArrayProperty = CastField<FArrayProperty>(Outer))
+		{
+			if (ArrayProperty->Inner != Property)
+			{
+				return nullptr;
+			}
+
+			FScriptArrayHelper ArrayHelper(ArrayProperty, ArrayProperty->template ContainerPtrToValuePtr<void>(Data));
+			int32 Index = ArrayHelper.AddValue();
+
+			return ArrayHelper.GetRawPtr(Index);
+		}
+
+		if (ArrayIndex >= Property->ArrayDim)
+		{
+			return nullptr;
+		}
+
+		return Property->template ContainerPtrToValuePtr<void>(Data, ArrayIndex);
+	}
+
+	/**
+	* Sets the value of the given property.
+	*
+	* @param Property The property to set.
+	* @param Outer The property that contains the property to be set, if any.
+	* @param Data A pointer to the memory holding the property's data.
+	* @param ArrayIndex The index of the element to set (if the property is an array).
+	* 
+	* @return true on success, false otherwise.
+	*/
+	template<typename PropertyType, typename ValueType>
+	static bool SetPropertyValue(PropertyType* Property, FProperty* Outer, void* Data, int32 ArrayIndex, const ValueType& Value)
+	{
+		if (void* Ptr = GetPropertyValuePtr(Property, Outer, Data, ArrayIndex))
+		{
+			Property->SetPropertyValue(Ptr, Value);
+			return true;
+		}
+
+		return false;
+	}
+
+	/**
+	 * Sets a numeric UProperty (int, uint, float, or double) on a given object or struct.
+	 *
+	 * @tparam T           Numeric type to write (e.g. float, double, int8, uint64, etc.).
+	 * @param Property     Target property.
+	 * @param Outer        The owning property of the target property.
+	 * @param Data         Pointer to the raw memory container holding the property value (object or struct).
+	 * @param ArrayIndex   Index within a TArray property. Use 0 for non-array properties.
+	 * @param InValue      The value to assign. Will be cast to the property's actual type.
+	 * 
+	 * @return             True if Property was a numeric type and the value was written. false otherwise.
+	 */
+	template<typename T>
+	static bool SetNumericPropertyValue(
+		FProperty* Property,
+		FProperty* Outer,
+		void* Data,
+		int32 ArrayIndex,
+		T InValue)
+	{
+		if (FNumericProperty* NumProp = CastField<FNumericProperty>(Property))
+		{
+			void* ValuePtr = GetPropertyValuePtr(Property, Outer, Data, ArrayIndex);
+
+			if (!ValuePtr)
+			{
+				return false;
+			}
+
+			if (NumProp->IsFloatingPoint())
+			{
+				NumProp->SetFloatingPointPropertyValue(ValuePtr, static_cast<double>(InValue));
+			}
+			else // integer or enum
+			{
+				if constexpr (std::is_signed_v<T>)
+				{
+					NumProp->SetIntPropertyValue(ValuePtr, static_cast<int64>(InValue));
+				}
+				else
+				{
+					NumProp->SetIntPropertyValue(ValuePtr, static_cast<uint64>(InValue));
+				}
+			}
+			return true;
+		}
+
+		return false;
+	}
+};
+
+
+FOpenTrackIOCborStructDeserializerBackend::FOpenTrackIOCborStructDeserializerBackend(FArchive& Archive)
+	: CborReader(&Archive, ECborEndianness::StandardCompliant)
+{}
+
+FOpenTrackIOCborStructDeserializerBackend::~FOpenTrackIOCborStructDeserializerBackend() = default;
+
+const FString& FOpenTrackIOCborStructDeserializerBackend::GetCurrentPropertyName() const
+{
+	return LastMapKey;
+}
+
+FString FOpenTrackIOCborStructDeserializerBackend::GetDebugString() const
+{
+	FArchive* Ar = const_cast<FArchive*>(CborReader.GetArchive());
+	return FString::Printf(TEXT("Offset: %" UINT64_FMT), Ar ? Ar->Tell() : 0);
+}
+
+const FString& FOpenTrackIOCborStructDeserializerBackend::GetLastErrorMessage() const
+{
+	static FString Dummy;
+	return Dummy;
+}
+
+bool FOpenTrackIOCborStructDeserializerBackend::GetNextToken(EStructDeserializerBackendTokens& OutToken)
+{
+	LastMapKey.Reset();
+
+	if (bDeserializingByteArray) // Deserializing the content of a TArray<uint8>/TArray<int8> property?
+	{
+		if (DeserializingByteArrayIndex < LastContext.AsByteArray().Num())
+		{
+			OutToken = EStructDeserializerBackendTokens::Property; // Need to consume a byte from the CBOR ByteString as a UByteProperty/UInt8Property.
+		}
+		else
+		{
+			bDeserializingByteArray = false;
+			OutToken = EStructDeserializerBackendTokens::ArrayEnd; // All bytes from the byte string were deserialized into the TArray<uint8>/TArray<int8>.
+		}
+
+		return true;
+	}
+
+	if (!CborReader.ReadNext(LastContext))
+	{
+		OutToken = LastContext.IsError() ? EStructDeserializerBackendTokens::Error : EStructDeserializerBackendTokens::None;
+		return false;
+	}
+
+	if (LastContext.IsBreak())
+	{
+		ECborCode ContainerEndType = LastContext.AsBreak();
+		// We do not support indefinite string container type
+		check(ContainerEndType == ECborCode::Array || ContainerEndType == ECborCode::Map);
+		OutToken = ContainerEndType == ECborCode::Array ? EStructDeserializerBackendTokens::ArrayEnd : EStructDeserializerBackendTokens::StructureEnd;
+		return true;
+	}
+
+	// if after reading the last context, the parent context is a map with an odd length, we just read a key
+	if (CborReader.GetContext().MajorType() == ECborCode::Map && (CborReader.GetContext().AsLength() & 1))
+	{
+		// Should be a string
+		check(LastContext.MajorType() == ECborCode::TextString);
+		LastMapKey = LastContext.AsString();
+
+		// Read next and carry on
+		if (!CborReader.ReadNext(LastContext))
+		{
+			OutToken = LastContext.IsError() ? EStructDeserializerBackendTokens::Error : EStructDeserializerBackendTokens::None;
+			return false;
+		}
+	}
+
+	switch (LastContext.MajorType())
+	{
+	case ECborCode::Array:
+		OutToken = EStructDeserializerBackendTokens::ArrayStart;
+		break;
+
+	case ECborCode::Map:
+		OutToken = EStructDeserializerBackendTokens::StructureStart;
+		break;
+
+	case ECborCode::ByteString: // Used for size optimization on TArray<uint8>/TArray<int8>. Might be replaced if https://datatracker.ietf.org/doc/draft-ietf-cbor-array-tags/ is adopted.
+		OutToken = EStructDeserializerBackendTokens::ArrayStart;
+		DeserializingByteArrayIndex = 0;
+		bDeserializingByteArray = true;
+		break;
+
+	case ECborCode::Int:
+		// fall through
+	case ECborCode::Uint:
+		// fall through
+	case ECborCode::TextString:
+		// fall through
+	case ECborCode::Prim:
+		OutToken = EStructDeserializerBackendTokens::Property;
+		break;
+
+	default:
+		// Other types are unsupported
+		return false;
+	}
+
+	return true;
+}
+
+bool FOpenTrackIOCborStructDeserializerBackend::ReadProperty(FProperty* Property, FProperty* Outer, void* Data, int32 ArrayIndex)
+{
+	switch (LastContext.MajorType())
+	{
+	// Unsigned Integers
+	case ECborCode::Uint:
+	{
+		const bool bWasSet = UE::OpenTrackIOCbor::Private::SetNumericPropertyValue(Property, Outer, Data, ArrayIndex, LastContext.AsUInt());
+
+		if (!bWasSet)
+		{
+			UE_LOG(LogLiveLinkOpenTrackIO, Verbose, TEXT("Unsigned integer field %s with value '%" UINT64_FMT "' is not supported in FProperty type %s (%s)"),
+				*Property->GetFName().ToString(), LastContext.AsUInt(), *Property->GetClass()->GetName(), *GetDebugString());
+		}
+
+		return bWasSet;
+	}
+
+	// Signed Integers
+	case ECborCode::Int:
+	{
+		const bool bWasSet = UE::OpenTrackIOCbor::Private::SetNumericPropertyValue(Property, Outer, Data, ArrayIndex, LastContext.AsInt());
+
+		if (!bWasSet)
+		{
+			UE_LOG(LogLiveLinkOpenTrackIO, Verbose, TEXT("Integer field %s with value '%" UINT64_FMT "' is not supported in FProperty type %s (%s)"),
+				*Property->GetFName().ToString(), LastContext.AsUInt(), *Property->GetClass()->GetName(), *GetDebugString());
+		}
+
+		return bWasSet;
+	}
+
+	// Strings, Names, Enumerations & Object/Class reference
+	case ECborCode::TextString:
+	{
+		if (FAnsiStrProperty* AnsiStrProperty = CastField<FAnsiStrProperty>(Property))
+		{
+			return UE::OpenTrackIOCbor::Private::SetPropertyValue(AnsiStrProperty, Outer, Data, ArrayIndex, LastContext.AsAnsiString());
+		}
+
+		if (FUtf8StrProperty* Utf8StrProperty = CastField<FUtf8StrProperty>(Property))
+		{
+			return UE::OpenTrackIOCbor::Private::SetPropertyValue(Utf8StrProperty, Outer, Data, ArrayIndex, LastContext.AsUtf8String());
+		}
+
+		const FString StringValue = LastContext.AsString();
+
+		if (FStrProperty* StrProperty = CastField<FStrProperty>(Property))
+		{
+			return UE::OpenTrackIOCbor::Private::SetPropertyValue(StrProperty, Outer, Data, ArrayIndex, StringValue);
+		}
+
+		if (FNameProperty* NameProperty = CastField<FNameProperty>(Property))
+		{
+			return UE::OpenTrackIOCbor::Private::SetPropertyValue(NameProperty, Outer, Data, ArrayIndex, FName(*StringValue));
+		}
+
+		if (FTextProperty* TextProperty = CastField<FTextProperty>(Property))
+		{
+			FText TextValue;
+			if (!FTextStringHelper::ReadFromBuffer(*StringValue, TextValue))
+			{
+				TextValue = FText::FromString(StringValue);
+			}
+			return UE::OpenTrackIOCbor::Private::SetPropertyValue(TextProperty, Outer, Data, ArrayIndex, TextValue);
+		}
+
+		if (FByteProperty* ByteProperty = CastField<FByteProperty>(Property))
+		{
+			if (!ByteProperty->Enum)
+			{
+				return false;
+			}
+
+			int64 Value = ByteProperty->Enum->GetValueByName(*StringValue);
+			if (Value == INDEX_NONE)
+			{
+				return false;
+			}
+
+			return UE::OpenTrackIOCbor::Private::SetPropertyValue(ByteProperty, Outer, Data, ArrayIndex, (uint8)Value);
+		}
+
+		if (FEnumProperty* EnumProperty = CastField<FEnumProperty>(Property))
+		{
+			int64 Value = EnumProperty->GetEnum()->GetValueByName(*StringValue);
+			if (Value == INDEX_NONE)
+			{
+				return false;
+			}
+
+			if (void* ElementPtr = UE::OpenTrackIOCbor::Private::GetPropertyValuePtr(EnumProperty, Outer, Data, ArrayIndex))
+			{
+				EnumProperty->GetUnderlyingProperty()->SetIntPropertyValue(ElementPtr, Value);
+				return true;
+			}
+
+			return false;
+		}
+
+		if (FClassProperty* ClassProperty = CastField<FClassProperty>(Property))
+		{
+			return UE::OpenTrackIOCbor::Private::SetPropertyValue(ClassProperty, Outer, Data, ArrayIndex, LoadObject<UClass>(NULL, *StringValue, NULL, LOAD_NoWarn));
+		}
+
+		if (FSoftClassProperty* SoftClassProperty = CastField<FSoftClassProperty>(Property))
+		{
+			return UE::OpenTrackIOCbor::Private::SetPropertyValue(SoftClassProperty, Outer, Data, ArrayIndex, FSoftObjectPtr(FSoftObjectPath(StringValue)));
+		}
+
+		if (FObjectProperty* ObjectProperty = CastField<FObjectProperty>(Property))
+		{
+			return UE::OpenTrackIOCbor::Private::SetPropertyValue(ObjectProperty, Outer, Data, ArrayIndex, StaticFindObject(ObjectProperty->PropertyClass, nullptr, *StringValue));
+		}
+
+		if (FWeakObjectProperty* WeakObjectProperty = CastField<FWeakObjectProperty>(Property))
+		{
+			return UE::OpenTrackIOCbor::Private::SetPropertyValue(WeakObjectProperty, Outer, Data, ArrayIndex, FWeakObjectPtr(StaticFindObject(WeakObjectProperty->PropertyClass, nullptr, *StringValue)));
+		}
+
+		if (FSoftObjectProperty* SoftObjectProperty = CastField<FSoftObjectProperty>(Property))
+		{
+			return UE::OpenTrackIOCbor::Private::SetPropertyValue(SoftObjectProperty, Outer, Data, ArrayIndex, FSoftObjectPtr(FSoftObjectPath(StringValue)));
+		}
+
+		UE_LOG(LogLiveLinkOpenTrackIO, Verbose, TEXT("String field %s with value '%s' is not supported in FProperty type %s (%s)"), *Property->GetFName().ToString(), *StringValue, *Property->GetClass()->GetName(), *GetDebugString());
+
+		return false;
+	}
+
+	// Stream of bytes: Used for TArray<uint8>/TArray<int8>
+	case ECborCode::ByteString:
+	{
+		check(bDeserializingByteArray);
+
+		// Consume one byte from the byte string.
+		TArrayView<const uint8> DeserializedByteArray = LastContext.AsByteArray();
+		check(DeserializingByteArrayIndex < DeserializedByteArray.Num());
+		uint8 ByteValue = DeserializedByteArray[DeserializingByteArrayIndex++];
+
+		if (FByteProperty* ByteProperty = CastField<FByteProperty>(Property))
+		{
+			return UE::OpenTrackIOCbor::Private::SetPropertyValue(ByteProperty, Outer, Data, ArrayIndex, ByteValue);
+		}
+		else if (FInt8Property* Int8Property = CastField<FInt8Property>(Property))
+		{
+			return UE::OpenTrackIOCbor::Private::SetPropertyValue(Int8Property, Outer, Data, ArrayIndex, (int8)ByteValue);
+		}
+
+		UE_LOG(LogLiveLinkOpenTrackIO, Verbose, TEXT("Error while deserializing field %s. Unexpected UProperty type %s. Expected a UByteProperty/UInt8Property to deserialize a TArray<uint8>/TArray<int8>"), *Property->GetFName().ToString(), *Property->GetClass()->GetName());
+		return false;
+	}
+
+	// Prim
+	case ECborCode::Prim:
+	{
+		switch (LastContext.AdditionalValue())
+		{
+		// Boolean
+		case ECborCode::True:
+			// fall through
+		case ECborCode::False:
+		{
+			const FCoreTexts& CoreTexts = FCoreTexts::Get();
+
+			if (FBoolProperty* BoolProperty = CastField<FBoolProperty>(Property))
+			{
+				return UE::OpenTrackIOCbor::Private::SetPropertyValue(BoolProperty, Outer, Data, ArrayIndex, LastContext.AsBool());
+			}
+			UE_LOG(LogLiveLinkOpenTrackIO, Verbose, TEXT("Boolean field %s with value '%s' is not supported in FProperty type %s (%s)"), *Property->GetFName().ToString(), LastContext.AsBool() ? *(CoreTexts.True.ToString()) : *(CoreTexts.False.ToString()), *Property->GetClass()->GetName(), *GetDebugString());
+			return false;
+		}
+
+		// Null
+		case ECborCode::Null:
+		{
+			return UE::OpenTrackIOCbor::Private::ClearPropertyValue(Property, Outer, Data, ArrayIndex);
+		}
+
+		// Float
+		case ECborCode::Value_4Bytes:
+		{
+			const bool bWasSet = UE::OpenTrackIOCbor::Private::SetNumericPropertyValue(Property, Outer, Data, ArrayIndex, LastContext.AsFloat());
+
+			if (!bWasSet)
+			{
+				UE_LOG(LogLiveLinkOpenTrackIO, Verbose, TEXT("Float field %s with value '%f' is not supported in FProperty type %s (%s)"),
+					*Property->GetFName().ToString(), LastContext.AsFloat(), *Property->GetClass()->GetName(), *GetDebugString());
+			}
+
+			return bWasSet;
+		}
+
+		// Double
+		case ECborCode::Value_8Bytes:
+		{
+			const bool bWasSet = UE::OpenTrackIOCbor::Private::SetNumericPropertyValue(Property, Outer, Data, ArrayIndex, LastContext.AsDouble());
+
+			if (!bWasSet)
+			{
+				UE_LOG(LogLiveLinkOpenTrackIO, Verbose, TEXT("Double field %s with value '%f' is not supported in FProperty type %s (%s)"),
+					*Property->GetFName().ToString(), LastContext.AsFloat(), *Property->GetClass()->GetName(), *GetDebugString());
+			}
+
+			return bWasSet;
+		}
+
+		default:
+			UE_LOG(LogLiveLinkOpenTrackIO, Verbose, TEXT("Unsupported primitive type for %s with value '%f' in FProperty type %s (%s)"), *Property->GetFName().ToString(), LastContext.AsDouble(), *Property->GetClass()->GetName(), *GetDebugString());
+			return false;
+		}
+	}
+	}
+
+	return true;
+}
+
+bool FOpenTrackIOCborStructDeserializerBackend::ReadPODArray(FArrayProperty* ArrayProperty, void* Data)
+{
+	// if we just read a byte array, copy the full array if the inner property is of the appropriate type 
+	if (bDeserializingByteArray
+		&& (CastField<FByteProperty>(ArrayProperty->Inner) || CastField<FInt8Property>(ArrayProperty->Inner)))
+	{
+		FScriptArrayHelper ArrayHelper(ArrayProperty, ArrayProperty->template ContainerPtrToValuePtr<void>(Data));
+		TArrayView<const uint8> DeserializedByteArray = LastContext.AsByteArray();
+		if (DeserializedByteArray.Num())
+		{
+			ArrayHelper.AddUninitializedValues(DeserializedByteArray.Num());
+			void* ArrayStart = ArrayHelper.GetRawPtr();
+			FMemory::Memcpy(ArrayStart, DeserializedByteArray.GetData(), DeserializedByteArray.Num());
+		}
+		bDeserializingByteArray = false;
+		return true;
+	}
+	return false;
+}
+
+void FOpenTrackIOCborStructDeserializerBackend::SkipArray()
+{
+	if (bDeserializingByteArray) // Deserializing a TArray<uint8>/TArray<int8> property as byte string?
+	{
+		check(DeserializingByteArrayIndex == 0);
+		bDeserializingByteArray = false;
+	}
+	else
+	{
+		CborReader.SkipContainer(ECborCode::Array);
+	}
+}
+
+void FOpenTrackIOCborStructDeserializerBackend::SkipStructure()
+{
+	CborReader.SkipContainer(ECborCode::Map);
+}
diff --git a/Engine/Plugins/VirtualProduction/LiveLinkOpenTrackIO/Source/LiveLinkOpenTrackIO/Private/OpenTrackIOCborStructDeserializerBackend.h b/Engine/Plugins/VirtualProduction/LiveLinkOpenTrackIO/Source/LiveLinkOpenTrackIO/Private/OpenTrackIOCborStructDeserializerBackend.h
new file mode 100644
--- /dev/null
+++ b/Engine/Plugins/VirtualProduction/LiveLinkOpenTrackIO/Source/LiveLinkOpenTrackIO/Private/OpenTrackIOCborStructDeserializerBackend.h
@@ -0,0 +1,59 @@
+// Copyright Epic Games, Inc. All Rights Reserved.
+
+#pragma once
+
+#include "CborReader.h"
+#include "CborTypes.h"
+#include "Containers/UnrealString.h"
+#include "CoreMinimal.h"
+#include "HAL/Platform.h"
+#include "IStructDeserializerBackend.h"
+
+class FArchive;
+class FArrayProperty;
+class FProperty;
+
+/**
+ * Implements a reader for UStruct deserialization using Cbor.
+ */
+class FOpenTrackIOCborStructDeserializerBackend
+	: public IStructDeserializerBackend
+{
+public:
+
+	/**
+	 * Creates and initializes a new instance.
+	 * 
+	 * @param Archive The archive to deserialize from.
+	 */
+	FOpenTrackIOCborStructDeserializerBackend(FArchive& Archive);
+	virtual ~FOpenTrackIOCborStructDeserializerBackend();
+
+public:
+
+	// IStructDeserializerBackend interface
+	virtual const FString& GetCurrentPropertyName() const override;
+	virtual FString GetDebugString() const override;
+	virtual const FString& GetLastErrorMessage() const override;
+	virtual bool GetNextToken(EStructDeserializerBackendTokens& OutToken) override;
+	virtual bool ReadProperty(FProperty* Property, FProperty* Outer, void* Data, int32 ArrayIndex) override;
+	virtual bool ReadPODArray(FArrayProperty* ArrayProperty, void* Data) override;
+	virtual void SkipArray() override;
+	virtual void SkipStructure() override;
+
+private:
+	/** Holds the Cbor reader used for the actual reading of the archive. */
+	FCborReader CborReader;
+
+	/** Holds the last read Cbor Context. */
+	FCborContext LastContext;
+
+	/** Holds the last map key. */
+	FString LastMapKey;
+
+	/** The index of the next byte to copy from the CBOR byte stream into the corresponding TArray<uint8>/TArray<int8> property. */
+	int32 DeserializingByteArrayIndex = 0;
+
+	/** Whether a TArray<uint8>/TArray<int8> property is being deserialized. */
+	bool bDeserializingByteArray = false;
+};
diff --git a/Engine/Plugins/VirtualProduction/LiveLinkOpenTrackIO/Source/LiveLinkOpenTrackIO/Public/LiveLinkOpenTrackIO.h b/Engine/Plugins/VirtualProduction/LiveLinkOpenTrackIO/Source/LiveLinkOpenTrackIO/Public/LiveLinkOpenTrackIO.h
--- a/Engine/Plugins/VirtualProduction/LiveLinkOpenTrackIO/Source/LiveLinkOpenTrackIO/Public/LiveLinkOpenTrackIO.h
+++ b/Engine/Plugins/VirtualProduction/LiveLinkOpenTrackIO/Source/LiveLinkOpenTrackIO/Public/LiveLinkOpenTrackIO.h
@@ -2,6 +2,7 @@
 
 #pragma once
 
+#include "Logging/LogMacros.h"
 #include "Modules/ModuleManager.h"
 
 DECLARE_LOG_CATEGORY_EXTERN(LogLiveLinkOpenTrackIO, Log, All);
