OpenTrackIO: Added support for optional numeric properties in cbor and json deserialization into structs.

This is used to set the data supported flags in the live link static data. That is, FLiveLinkLensStaticData::bIsFocalLengthSupported et al.

#rnx

diff --git a/Engine/Plugins/VirtualProduction/LiveLinkOpenTrackIO/Source/LiveLinkOpenTrackIO/Private/LiveLinkOpenTrackIOConversions.h b/Engine/Plugins/VirtualProduction/LiveLinkOpenTrackIO/Source/LiveLinkOpenTrackIO/Private/LiveLinkOpenTrackIOConversions.h
--- a/Engine/Plugins/VirtualProduction/LiveLinkOpenTrackIO/Source/LiveLinkOpenTrackIO/Private/LiveLinkOpenTrackIOConversions.h
+++ b/Engine/Plugins/VirtualProduction/LiveLinkOpenTrackIO/Source/LiveLinkOpenTrackIO/Private/LiveLinkOpenTrackIOConversions.h
@@ -33,7 +33,7 @@ namespace LiveLinkOpenTrackIOConversions
 
 
 	/** Populate Live Link Lens Frame Data from relevant OpenTrackIO structures */
-	inline void ConvertLens(
+	inline void ToUnrealLens(
 		FLiveLinkLensFrameData& OutUnrealLensData, 
 		const FLiveLinkOpenTrackIOLens* InLensData,
 		const FLiveLinkOpenTrackIOStaticCamera* InCamera
@@ -41,15 +41,33 @@ namespace LiveLinkOpenTrackIOConversions
 	{
 		if (InLensData)
 		{
-			OutUnrealLensData.FocusDistance = InLensData->FocusDistance * MetersToCentimeters;
-			OutUnrealLensData.Aperture = InLensData->FStop;
-			OutUnrealLensData.FocalLength = InLensData->PinholeFocalLength;  // Both are in mm
+			if (InLensData->FocusDistance.IsSet())
+			{
+				OutUnrealLensData.FocusDistance = InLensData->FocusDistance.GetValue() * MetersToCentimeters;
+			}
+
+			if (InLensData->FStop.IsSet())
+			{
+				OutUnrealLensData.Aperture = InLensData->FStop.GetValue();
+			}
+
+			if (InLensData->PinholeFocalLength.IsSet())
+			{
+				OutUnrealLensData.FocalLength = InLensData->PinholeFocalLength.GetValue();  // Both are in mm
+			}
 		}
 
 		if (InCamera)
 		{
-			OutUnrealLensData.FilmBackHeight = InCamera->ActiveSensorPhysicalDimensions.Height; // Both are in mm
-			OutUnrealLensData.FilmBackWidth = InCamera->ActiveSensorPhysicalDimensions.Width;   // Both are in mm
+			if (InCamera->ActiveSensorPhysicalDimensions.Height.IsSet())
+			{
+				OutUnrealLensData.FilmBackHeight = InCamera->ActiveSensorPhysicalDimensions.Height.GetValue(); // Both are in mm
+			}
+
+			if (InCamera->ActiveSensorPhysicalDimensions.Width.IsSet())
+			{
+				OutUnrealLensData.FilmBackWidth = InCamera->ActiveSensorPhysicalDimensions.Width.GetValue();   // Both are in mm
+			}
 
 			if (InLensData)
 			{
@@ -63,30 +81,30 @@ namespace LiveLinkOpenTrackIOConversions
 	}
 
 	/** From OpenTrackIO to Unreal Translation: Swap X <-> Y and convert meters to cm */
-	FVector ConvertTranslation(const FLiveLinkOpenTrackIO_XYZ& InXYZ)
+	FVector ToUnrealTranslation(const FLiveLinkOpenTrackIO_XYZ& InXYZ)
 	{
 		return FVector(InXYZ.Y, InXYZ.X, InXYZ.Z) * MetersToCentimeters;
 	}
 
 	/** From OpenTrackIO to Unreal Scale: Swap X <-> Y (unitless) */
-	FVector ConvertScale(const FLiveLinkOpenTrackIO_XYZ& InXYZ)
+	FVector ToUnrealScale(const FLiveLinkOpenTrackIO_XYZ& InXYZ)
 	{
 		return FVector(InXYZ.Y, InXYZ.X, InXYZ.Z);
 	}
 
 	/** From OpenTrackIO to Unreal Translation: Invert sign of Yaw */
-	FRotator ConvertRotation(const FLiveLinkOpenTrackIO_Rotator& InRotator)
+	FRotator ToUnrealRotation(const FLiveLinkOpenTrackIO_Rotator& InRotator)
 	{
 		return FRotator(InRotator.Tilt, -InRotator.Pan, InRotator.Roll);
 	}
 
 	/** From OpenTrackIO to Unreal FTransform */
-	FTransform ConvertTransform(const FLiveLinkOpenTrackIOTransform& InTransform)
+	FTransform ToUnrealTransform(const FLiveLinkOpenTrackIOTransform& InTransform)
 	{
 		return FTransform(
-			ConvertRotation(InTransform.Rotation),
-			ConvertTranslation(InTransform.Translation),
-			ConvertScale(InTransform.Scale)
+			ToUnrealRotation(InTransform.Rotation),
+			ToUnrealTranslation(InTransform.Translation),
+			ToUnrealScale(InTransform.Scale)
 		);
 	}
 }
diff --git a/Engine/Plugins/VirtualProduction/LiveLinkOpenTrackIO/Source/LiveLinkOpenTrackIO/Private/LiveLinkOpenTrackIOParser.cpp b/Engine/Plugins/VirtualProduction/LiveLinkOpenTrackIO/Source/LiveLinkOpenTrackIO/Private/LiveLinkOpenTrackIOParser.cpp
--- a/Engine/Plugins/VirtualProduction/LiveLinkOpenTrackIO/Source/LiveLinkOpenTrackIO/Private/LiveLinkOpenTrackIOParser.cpp
+++ b/Engine/Plugins/VirtualProduction/LiveLinkOpenTrackIO/Source/LiveLinkOpenTrackIO/Private/LiveLinkOpenTrackIOParser.cpp
@@ -23,7 +23,82 @@
 
 namespace UE::OpenTrackIO::Private
 {
-	
+
+/**
+* Callback for importing JSON values into TOptional (currently numeric) properties.
+* If the property is not an FOptionalProperty or does not wrap a numeric type, 
+* the callback returns false to let the default parser take over.
+*
+* @param JsonValue   The JSON value to import.
+* @param Property    The TOptional property to write into.
+* @param Value       Pointer to memory holding the optional property value.
+* 
+* @return            True if this callback handled the import (optional set successfully).
+*/
+bool JsonOptionalImporterCallback(const TSharedPtr<FJsonValue>& JsonValue, FProperty* Property, void* Value)
+{
+	// Only handle optionals
+
+	FOptionalProperty* OptionalProperty = CastField<FOptionalProperty>(Property);
+
+	if (!OptionalProperty)
+	{
+		return false;
+	}
+
+	// Handle numeric optional
+
+	FNumericProperty* NumericProperty = CastField<FNumericProperty>(OptionalProperty->GetValueProperty());
+
+	if (!NumericProperty)
+	{
+		return false;
+	}
+
+	// Handle floats and double property types.
+
+	if (NumericProperty->IsFloatingPoint())
+	{
+		void* OptionalValue = OptionalProperty->MarkSetAndGetInitializedValuePointerToReplace(Value);
+
+		if (!OptionalValue)
+		{
+			return false;
+		}
+
+		NumericProperty->SetFloatingPointPropertyValue(OptionalValue, JsonValue->AsNumber());
+
+		return true;
+	}
+		
+	// Handle integer property types.
+
+	if (NumericProperty->IsInteger())
+	{
+		void* OptionalValue = OptionalProperty->MarkSetAndGetInitializedValuePointerToReplace(Value);
+
+		if (!OptionalValue)
+		{
+			return false;
+		}
+
+		if (JsonValue->Type == EJson::String)
+		{
+			// parse string -> int64 ourselves so we don't lose any precision going through AsNumber (aka double)
+			NumericProperty->SetIntPropertyValue(OptionalValue, FCString::Atoi64(*JsonValue->AsString()));
+		}
+		else
+		{
+			NumericProperty->SetIntPropertyValue(OptionalValue, (int64)JsonValue->AsNumber());
+		}
+
+		return true;
+	}
+
+	// We did not write into this type, so we return false for the default parser to handle it.
+	return false;
+}
+
 TOptional<FLiveLinkOpenTrackIOData> ParseJsonBlob(const FString& JsonBlob)
 {
 	TSharedPtr<FJsonObject> JsonObject;
@@ -47,7 +122,12 @@ TOptional<FLiveLinkOpenTrackIOData> ParseJsonBlob(const FString& JsonBlob)
 	}
 	
 	FLiveLinkOpenTrackIOData Data;
-	if (FJsonObjectConverter::JsonObjectStringToUStruct(JsonBlob, &Data))
+
+	// We use the json importer callback to add support of optional properties.
+	FJsonObjectConverter::CustomImportCallback OptionalNumericImporter;
+	OptionalNumericImporter.BindStatic(JsonOptionalImporterCallback);
+
+	if (FJsonObjectConverter::JsonObjectStringToUStruct(JsonBlob, &Data, 0, 0, false, nullptr, &OptionalNumericImporter))
 	{
 		OutData.Emplace(MoveTemp(Data));
 	}
@@ -212,7 +292,7 @@ TOptional<FOpenTrackIOHeaderWithData> ParseOpenTrackIOPacket(TArrayView<uint8> B
 	uint16_t Checksum = CalculateChecksum(HeaderMinusChecksum, Payload);
 	if (Checksum != Header.Checksum)
 	{
-		UE_LOG(LogLiveLinkOpenTrackIO, Warning, TEXT("Failed to verify packet cheksum."));
+		UE_LOG(LogLiveLinkOpenTrackIO, Warning, TEXT("Failed to verify packet checksum."));
 		return OutData;
 	}
 
diff --git a/Engine/Plugins/VirtualProduction/LiveLinkOpenTrackIO/Source/LiveLinkOpenTrackIO/Private/LiveLinkOpenTrackIOTranscoder.cpp b/Engine/Plugins/VirtualProduction/LiveLinkOpenTrackIO/Source/LiveLinkOpenTrackIO/Private/LiveLinkOpenTrackIOTranscoder.cpp
--- a/Engine/Plugins/VirtualProduction/LiveLinkOpenTrackIO/Source/LiveLinkOpenTrackIO/Private/LiveLinkOpenTrackIOTranscoder.cpp
+++ b/Engine/Plugins/VirtualProduction/LiveLinkOpenTrackIO/Source/LiveLinkOpenTrackIO/Private/LiveLinkOpenTrackIOTranscoder.cpp
@@ -75,16 +75,40 @@ TOptional<FLiveLinkStaticDataStruct> FLiveLinkOpenTrackIOCache::TryAddCamera(con
 FLiveLinkStaticDataStruct FLiveLinkOpenTrackIOCache::MakeStaticData(const FLiveLinkOpenTrackIOData& Data)
 {
 	/** Lens class inherits from the camera data so we use that as our base for the static struct. */
+
 	FLiveLinkStaticDataStruct StaticDataStruct(FLiveLinkLensStaticData::StaticStruct());
 	FLiveLinkLensStaticData& LensData = *StaticDataStruct.Cast<FLiveLinkLensStaticData>();
 
-	LensData.bIsFocalLengthSupported = true;
-	LensData.bIsApertureSupported = true;
-	LensData.bIsFocusDistanceSupported = true;
-	LensData.bIsDynamicFilmbackSupported = true;
+	if (Data.Lens.PinholeFocalLength.IsSet())
+	{
+		LensData.bIsFocalLengthSupported = true;
+	}
+
+	if (Data.Lens.FStop.IsSet())
+	{
+		LensData.bIsApertureSupported = true;
+	}
 
-	LensData.FilmBackHeight = Data.Static.Camera.ActiveSensorPhysicalDimensions.Height;
-	LensData.FilmBackWidth = Data.Static.Camera.ActiveSensorPhysicalDimensions.Width;
+	if (Data.Lens.FocusDistance.IsSet())
+	{
+		LensData.bIsFocusDistanceSupported = true;
+	}
+
+	if (Data.Static.Camera.ActiveSensorPhysicalDimensions.Height.IsSet()
+		&& Data.Static.Camera.ActiveSensorPhysicalDimensions.Width.IsSet())
+	{
+		LensData.bIsDynamicFilmbackSupported = true;
+
+		LensData.FilmBackHeight = Data.Static.Camera.ActiveSensorPhysicalDimensions.Height.GetValue();
+		LensData.FilmBackWidth = Data.Static.Camera.ActiveSensorPhysicalDimensions.Width.GetValue();
+	}
+
+	if (Data.Transforms.Num() > 0 || Data.Lens.EntrancePupilOffset.IsSet())
+	{
+		LensData.bIsLocationSupported = true;
+		LensData.bIsRotationSupported = true;
+		LensData.bIsScaleSupported = true;
+	}
 
 	/** OpenTrackIO only allows for spherical models right now. */
 	LensData.LensModel = FName(TEXT("Spherical"));
@@ -103,16 +127,17 @@ TArray<FSubjectAndFrameDataTuple> FLiveLinkOpenTrackIOCache::TryReadFrameData(co
 	FTransform FinalTransform; 
 	for (const FLiveLinkOpenTrackIOTransform& Transform : Data.Transforms)
 	{
-		FinalTransform *= ConvertTransform(Transform);
+		FinalTransform *= ToUnrealTransform(Transform);
 	}
 
 	// Apply entrance pupil offset as a final transform in the forward direction.
+	if (Data.Lens.EntrancePupilOffset.IsSet())
 	{
-		const FLiveLinkOpenTrackIO_XYZ EntrancePupilAsXYZ(0, Data.Lens.EntrancePupilOffset, 0);
+		const FLiveLinkOpenTrackIO_XYZ EntrancePupilAsXYZ(0, Data.Lens.EntrancePupilOffset.GetValue(), 0);
 
 		const FTransform PupilOffset(
 			FQuat::Identity,
-			ConvertTranslation(EntrancePupilAsXYZ),
+			ToUnrealTranslation(EntrancePupilAsXYZ),
 			FVector(1, 1, 1)
 		);
 
@@ -140,7 +165,7 @@ TArray<FSubjectAndFrameDataTuple> FLiveLinkOpenTrackIOCache::TryReadFrameData(co
 	FLiveLinkFrameDataStruct FrameDataStruct = FLiveLinkFrameDataStruct(FLiveLinkLensFrameData::StaticStruct());
 	FLiveLinkLensFrameData& CameraAndLensData = *FrameDataStruct.Cast<FLiveLinkLensFrameData>();
 
-	ConvertLens(CameraAndLensData, &Data.Lens, Camera);
+	ToUnrealLens(CameraAndLensData, &Data.Lens, Camera);
 
 	CameraAndLensData.Transform = FinalTransform;
 	CameraAndLensData.MetaData.SceneTime = SampleTC;
diff --git a/Engine/Plugins/VirtualProduction/LiveLinkOpenTrackIO/Source/LiveLinkOpenTrackIO/Private/LiveLinkOpenTrackIOTypes.h b/Engine/Plugins/VirtualProduction/LiveLinkOpenTrackIO/Source/LiveLinkOpenTrackIO/Private/LiveLinkOpenTrackIOTypes.h
--- a/Engine/Plugins/VirtualProduction/LiveLinkOpenTrackIO/Source/LiveLinkOpenTrackIO/Private/LiveLinkOpenTrackIOTypes.h
+++ b/Engine/Plugins/VirtualProduction/LiveLinkOpenTrackIO/Source/LiveLinkOpenTrackIO/Private/LiveLinkOpenTrackIOTypes.h
@@ -107,10 +107,10 @@ struct FLiveLinkOpenTrackIOSensorDimensions
 	bool operator!=(const FLiveLinkOpenTrackIOSensorDimensions&) const = default;
 	
 	UPROPERTY(EditAnywhere, BlueprintReadWrite, Category = "OpenTrackIO")
-	double Height = 43.0f;
+	TOptional<double> Height;
 
 	UPROPERTY(EditAnywhere, BlueprintReadWrite, Category = "OpenTrackIO")
-	double Width = 24.0f;
+	TOptional<double> Width;
 };
 
 /**
@@ -694,8 +694,8 @@ struct FLiveLinkOpenTrackIOTimecode
 
 	FTimecode GetTimecode() const
 	{
-		double FPS = FrameRate.GetFrameRate().AsDecimal();
-		float SubFrameAsFloat = static_cast<float>(static_cast<double>(SubFrame) / FPS);
+		const double FPS = FrameRate.GetFrameRate().AsDecimal();
+		const float SubFrameAsFloat = static_cast<float>(static_cast<double>(SubFrame) / FPS);
 		return FTimecode(Hours, Minutes, Seconds, Frames, SubFrameAsFloat, false);
 	}
 	
@@ -976,7 +976,7 @@ struct FLiveLinkOpenTrackIOLens
 	 * virtual camera’s transform chain.
 	 */
 	UPROPERTY(EditAnywhere, BlueprintReadWrite, Category = "OpenTrackIO")
-	double  EntrancePupilOffset = 0;
+	TOptional<double> EntrancePupilOffset;
 
 	/**
 	 * Coefficients for calculating the exposure fall-off (vignetting) of a lens
@@ -986,15 +986,15 @@ struct FLiveLinkOpenTrackIOLens
 
 	/** The linear f-number of the lens, equal to the focal length divided by the diameter of the entrance pupil. */
 	UPROPERTY(EditAnywhere, BlueprintReadWrite, Category = "OpenTrackIO")
-	double FStop = 2.8f;
+	TOptional<double> FStop;
 
 	/** Distance between the pinhole and the image plane in the simple CGI pinhole camera model. (millimeters) */
 	UPROPERTY(EditAnywhere, BlueprintReadWrite, Category = "OpenTrackIO")
-	double PinholeFocalLength = 35.0f;
+	TOptional<double> PinholeFocalLength;
 
 	/** Focus distance/position of the lens in meters. */ 
 	UPROPERTY(EditAnywhere, BlueprintReadWrite, Category = "OpenTrackIO")
-	double FocusDistance = 10.0f;
+	TOptional<double> FocusDistance;
 
 	/**  Offset in x and y of the centre of perspective projection of the virtual camera */
 	UPROPERTY(EditAnywhere, BlueprintReadWrite, Category = "OpenTrackIO")
diff --git a/Engine/Plugins/VirtualProduction/LiveLinkOpenTrackIO/Source/LiveLinkOpenTrackIO/Private/OpenTrackIOCborStructDeserializerBackend.cpp b/Engine/Plugins/VirtualProduction/LiveLinkOpenTrackIO/Source/LiveLinkOpenTrackIO/Private/OpenTrackIOCborStructDeserializerBackend.cpp
--- a/Engine/Plugins/VirtualProduction/LiveLinkOpenTrackIO/Source/LiveLinkOpenTrackIO/Private/OpenTrackIOCborStructDeserializerBackend.cpp
+++ b/Engine/Plugins/VirtualProduction/LiveLinkOpenTrackIO/Source/LiveLinkOpenTrackIO/Private/OpenTrackIOCborStructDeserializerBackend.cpp
@@ -7,6 +7,7 @@
 #include "UObject/Class.h"
 #include "UObject/AnsiStrProperty.h"
 #include "UObject/EnumProperty.h"
+#include "UObject/PropertyOptional.h"
 #include "UObject/TextProperty.h"
 #include "UObject/UnrealType.h"
 #include "UObject/Utf8StrProperty.h"
@@ -262,6 +263,45 @@ bool FOpenTrackIOCborStructDeserializerBackend::GetNextToken(EStructDeserializer
 
 bool FOpenTrackIOCborStructDeserializerBackend::ReadProperty(FProperty* Property, FProperty* Outer, void* Data, int32 ArrayIndex)
 {
+	// Unwrap TOptional and recurse back into this function
+	if (FOptionalProperty* OptionalProperty = CastField<FOptionalProperty>(Property))
+	{
+		void* ValuePtr = UE::OpenTrackIOCbor::Private::GetPropertyValuePtr(OptionalProperty, Outer, Data, ArrayIndex);
+
+		if (!ValuePtr)
+		{
+			return false;
+		}
+
+		// "Set" the optional and get a pointer to its value
+
+		void* InnerData = OptionalProperty->MarkSetAndGetInitializedValuePointerToReplace(ValuePtr);
+
+		if (!InnerData)
+		{
+			return false;
+		}
+
+		// Get the property inside the optional
+		FProperty* ValueProp = OptionalProperty->GetValueProperty();
+
+		if (!ValueProp)
+		{
+			return false;
+		}
+
+		// Recurse back into this function to handle the inner property
+		const bool bInnerPropertyWasSet = ReadProperty(ValueProp, Property, InnerData, ArrayIndex);
+
+		// But if the inner property wasn't read into, then reset the optional
+		if (!bInnerPropertyWasSet)
+		{
+			OptionalProperty->MarkUnset(ValuePtr);
+		}
+
+		return bInnerPropertyWasSet;
+	}
+
 	switch (LastContext.MajorType())
 	{
 	// Unsigned Integers
