OpenTrackIO: Added subframes per frame.

OpenTrackIO doesn't specify subframes per frame, so it is inferred from the ratio between sample rate and timecode rate. We use this to convert from the subFrame count that OpenTrackIO sends to the frame fraction that Unreal's timecode representation uses, that can be used to accurately key the data on a timeline.

While doing this, I noticed that the FTimecode subframe was being dropped when initializing FQualifiedFrameTime, which prevents proper timecode usage and roundtripping TC -> QFT -> TC. This was fixed in the CL as well.

#rnx

diff --git a/Engine/Plugins/VirtualProduction/LiveLinkOpenTrackIO/Source/LiveLinkOpenTrackIO/Private/LiveLinkOpenTrackIOParser.cpp b/Engine/Plugins/VirtualProduction/LiveLinkOpenTrackIO/Source/LiveLinkOpenTrackIO/Private/LiveLinkOpenTrackIOParser.cpp
--- a/Engine/Plugins/VirtualProduction/LiveLinkOpenTrackIO/Source/LiveLinkOpenTrackIO/Private/LiveLinkOpenTrackIOParser.cpp
+++ b/Engine/Plugins/VirtualProduction/LiveLinkOpenTrackIO/Source/LiveLinkOpenTrackIO/Private/LiveLinkOpenTrackIOParser.cpp
@@ -248,10 +248,13 @@ inline FArchive& operator<<(FArchive& Ar, FLiveLinkOpenTrackIODatagramHeader& He
 
 TOptional<FOpenTrackIOHeaderWithData> ParseOpenTrackIOPacket(TArrayView<uint8> Bytes)
 {
+	TRACE_CPUPROFILER_EVENT_SCOPE(OpenTrackIO::ParseOpenTrackIOPacket);
+
 	const uint32 HeaderSize = sizeof(FLiveLinkOpenTrackIODatagramHeader);
 
 	TOptional<FOpenTrackIOHeaderWithData> OutData;
-	/** No way this ia a valid datagram if we can't fit it in our header. */
+
+	/** No way this is a valid datagram if we can't fit it in our header. */
 	if (Bytes.Num() < HeaderSize)
 	{
 		return OutData;
@@ -266,7 +269,10 @@ TOptional<FOpenTrackIOHeaderWithData> ParseOpenTrackIOPacket(TArrayView<uint8> B
 		Identifier = ByteSwap(Identifier);
 		if (Identifier != OpenTrackIOHeaderId)
 		{
-			UE_LOG(LogLiveLinkOpenTrackIO, Error, TEXT("Invalid OpenTrackIO header."));
+			UE_LOG(LogLiveLinkOpenTrackIO, Error, TEXT("Invalid OpenTrackIO header. It was %08X instead of %08X"),
+				Identifier, OpenTrackIOHeaderId
+			);
+
 			return OutData;
 		}
 	}
@@ -318,6 +324,16 @@ TOptional<FOpenTrackIOHeaderWithData> ParseOpenTrackIOPacket(TArrayView<uint8> B
 	{
 		if (Data->Protocol.IsSupported())
 		{
+			// Manually compensate for OpenTrackIO not explicitly stating the number of frames per frame.
+			{
+				// Calculate from sampleRate / timecodeRate, because it tells us how many samples there are per timecode frame.
+				const int32 RatioNum = Data->Timing.SampleRate.Num * Data->Timing.Timecode.FrameRate.Denom;
+				const int32 RatioDen = FMath::Max(1, Data->Timing.SampleRate.Denom * Data->Timing.Timecode.FrameRate.Num);
+
+				// Enforce a minium of 1 SubframesPerFrame, which is equivalent to no frame subdivision.
+				Data->Timing.Timecode.SubframesPerFrame = FMath::Max(1, FMath::DivideAndRoundNearest(RatioNum, RatioDen));
+			}
+
 			OutData.Emplace(FOpenTrackIOHeaderWithData{Header, *Data});
 		}
 		else
diff --git a/Engine/Plugins/VirtualProduction/LiveLinkOpenTrackIO/Source/LiveLinkOpenTrackIO/Private/LiveLinkOpenTrackIOSource.cpp b/Engine/Plugins/VirtualProduction/LiveLinkOpenTrackIO/Source/LiveLinkOpenTrackIO/Private/LiveLinkOpenTrackIOSource.cpp
--- a/Engine/Plugins/VirtualProduction/LiveLinkOpenTrackIO/Source/LiveLinkOpenTrackIO/Private/LiveLinkOpenTrackIOSource.cpp
+++ b/Engine/Plugins/VirtualProduction/LiveLinkOpenTrackIO/Source/LiveLinkOpenTrackIO/Private/LiveLinkOpenTrackIOSource.cpp
@@ -297,6 +297,8 @@ void FLiveLinkOpenTrackIOSource::HandleInboundData(const TSharedPtr<FArrayReader
 
 bool FLiveLinkOpenTrackIOSource::ConsumeInboundData()
 {
+	TRACE_CPUPROFILER_EVENT_SCOPE(OpenTrackIO::ConsumeInboundData);
+
 	bool bProcessedData = false;
 	FInboundMessage Message;
 	while (Inbound.Dequeue(Message))
@@ -313,8 +315,10 @@ bool FLiveLinkOpenTrackIOSource::ConsumeInboundData()
 			bProcessedData = true;
 			OpenTrackIOPacketCount++;
 
-			
-			PushDataToLiveLink_AnyThread(ParsedData->Get<0>() /** Header */, ParsedData->Get<1>() /** Body */);
+			const FLiveLinkOpenTrackIODatagramHeader& Header = ParsedData->Get<0>();
+			FLiveLinkOpenTrackIOData& Data = ParsedData->Get<1>();
+
+			PushDataToLiveLink_AnyThread(Header, Data);
 		}
 		else
 		{
@@ -326,6 +330,8 @@ bool FLiveLinkOpenTrackIOSource::ConsumeInboundData()
 
 void FLiveLinkOpenTrackIOSource::PushDataToLiveLink_AnyThread(const FLiveLinkOpenTrackIODatagramHeader& Header, const FLiveLinkOpenTrackIOData& InData)
 {
+	TRACE_CPUPROFILER_EVENT_SCOPE(OpenTrackIO::PushDataToLiveLink_AnyThread);
+
 	FLiveLinkOpenTrackIOCache* Cache;
 	if (OpenTrackIOCacheMap.Contains(InData.SourceId))
 	{
@@ -348,7 +354,7 @@ void FLiveLinkOpenTrackIOSource::PushDataToLiveLink_AnyThread(const FLiveLinkOpe
 		const FName SubjectName = Cache->GetSubjectNameFromData(InData);
 
 		Client->PushSubjectStaticData_AnyThread({ SourceGuid, SubjectName }, ULiveLinkLensRole::StaticClass(), MoveTemp(*StaticData));
-
+			  
 		// Remove the interim live link subject, if any.
 		if (!Cache->DefaultSubjectName.IsNone())
 		{
diff --git a/Engine/Plugins/VirtualProduction/LiveLinkOpenTrackIO/Source/LiveLinkOpenTrackIO/Private/LiveLinkOpenTrackIOTranscoder.cpp b/Engine/Plugins/VirtualProduction/LiveLinkOpenTrackIO/Source/LiveLinkOpenTrackIO/Private/LiveLinkOpenTrackIOTranscoder.cpp
--- a/Engine/Plugins/VirtualProduction/LiveLinkOpenTrackIO/Source/LiveLinkOpenTrackIO/Private/LiveLinkOpenTrackIOTranscoder.cpp
+++ b/Engine/Plugins/VirtualProduction/LiveLinkOpenTrackIO/Source/LiveLinkOpenTrackIO/Private/LiveLinkOpenTrackIOTranscoder.cpp
@@ -32,11 +32,13 @@ bool FLiveLinkOpenTrackIOCache::IsPacketInSequence(uint16 SequenceNumber, const
 	// Initially the expected sequence number is -1, which indicates we have not received our first packet yet. 
 	if (ExpectedSequenceNumber > 0 && SequenceNumber < ExpectedSequenceNumber)
 	{
-		double RateAsDecimal = InRate.AsDecimal();
-		check(RateAsDecimal>0);
-		const double SampleRateExpiry = 3 * 1.0 / InRate.AsDecimal();
-		// We we have exceeded a timeout period then consider the sequence stream as being reset.  In other words,
-		// sequence numbers are sequential only 
+		// Ensure a minimum rate. Also avoids divide by zero if the external producer sends zero rate.
+		constexpr double MinimumPeriod = 1.0 / 120.0;
+		const double Period = static_cast<double>(InRate.Denominator) / FMath::Max(1, InRate.Numerator);
+		const double SampleRateExpiry = 3 * FMath::Max(MinimumPeriod, Period);
+
+		// We we have exceeded a timeout period then consider the sequence stream as being reset.
+		// In other words, sequence numbers are sequential only.
 		if ((FPlatformTime::Seconds() - LastDataReceiveTimeInSeconds) < SampleRateExpiry)
 		{
 			return false;
diff --git a/Engine/Plugins/VirtualProduction/LiveLinkOpenTrackIO/Source/LiveLinkOpenTrackIO/Private/LiveLinkOpenTrackIOTypes.h b/Engine/Plugins/VirtualProduction/LiveLinkOpenTrackIO/Source/LiveLinkOpenTrackIO/Private/LiveLinkOpenTrackIOTypes.h
--- a/Engine/Plugins/VirtualProduction/LiveLinkOpenTrackIO/Source/LiveLinkOpenTrackIO/Private/LiveLinkOpenTrackIOTypes.h
+++ b/Engine/Plugins/VirtualProduction/LiveLinkOpenTrackIO/Source/LiveLinkOpenTrackIO/Private/LiveLinkOpenTrackIOTypes.h
@@ -34,7 +34,6 @@ struct FLiveLinkOpenTrackIOFrameRate
 
 	FFrameRate GetFrameRate() const
 	{
-		check(Num > 0);
 		return FFrameRate(static_cast<uint32>(Num), static_cast<uint32>(Denom));
 	}
 
@@ -691,17 +690,21 @@ struct FLiveLinkOpenTrackIOTimecode
 	
 	bool operator==(const FLiveLinkOpenTrackIOTimecode&) const = default;
 	bool operator!=(const FLiveLinkOpenTrackIOTimecode&) const = default;
-
-	FTimecode GetTimecode() const
-	{
-		const double FPS = FrameRate.GetFrameRate().AsDecimal();
-		const float SubFrameAsFloat = static_cast<float>(static_cast<double>(SubFrame) / FPS);
-		return FTimecode(Hours, Minutes, Seconds, Frames, SubFrameAsFloat, false);
-	}
 	
 	FQualifiedFrameTime GetQualifiedFrameTime() const
 	{
-		return FQualifiedFrameTime(GetTimecode(), FrameRate.GetFrameRate());
+		const double SubFrameFraction = FMath::Clamp<double>(
+			static_cast<double>(SubFrame) / FMath::Max(1, SubframesPerFrame), 
+			0.0, 
+			1.0
+		);
+
+		const bool bDropFrameFlag = FTimecode::UseDropFormatTimecode(FrameRate.GetFrameRate());
+
+		return FQualifiedFrameTime(
+			FTimecode(Hours, Minutes, Seconds, Frames, static_cast<float>(SubFrameFraction), bDropFrameFlag),
+			FrameRate.GetFrameRate()
+		);
 	}
 	
 	UPROPERTY(EditAnywhere, BlueprintReadWrite, Category = "OpenTrackIO")	
@@ -721,6 +724,10 @@ struct FLiveLinkOpenTrackIOTimecode
 	
 	UPROPERTY(EditAnywhere, BlueprintReadWrite, Category = "OpenTrackIO")
 	int32 SubFrame = 0;
+
+	/** Number of subframes per frame. 1 is the minimum */
+	UPROPERTY(EditAnywhere, BlueprintReadWrite, Category = "OpenTrackIO")
+	int32 SubframesPerFrame = 1;
 };
 
 USTRUCT(BlueprintType)
diff --git a/Engine/Source/Runtime/Core/Public/Misc/QualifiedFrameTime.h b/Engine/Source/Runtime/Core/Public/Misc/QualifiedFrameTime.h
--- a/Engine/Source/Runtime/Core/Public/Misc/QualifiedFrameTime.h
+++ b/Engine/Source/Runtime/Core/Public/Misc/QualifiedFrameTime.h
@@ -30,7 +30,7 @@ struct FQualifiedFrameTime
 	 * User construction from a timecode and its frame rate
 	 */
 	FQualifiedFrameTime(const FTimecode& InTimecode, const FFrameRate& InRate)
-		: Time(InTimecode.ToFrameNumber(InRate))
+		: Time(InTimecode.ToFrameNumber(InRate), InTimecode.Subframe)
 		, Rate(InRate)
 	{
 	}
