OpenTrackIO: Allow for static data every once in a while instead of either never or every frame.

The cache is now per sourceId_sourceNumber, which should allow for multiple sources or multiple source ids from the same ip address to be handled.

The default Subject name now includes the last 6 digits of the source id, to avoid a conflict when multiple source ids send the same source number. In the future we could allow the user to customize it instead.

We now  the source id from OTrIO steam to name the subject instead of the one used to create the live lnk source, because they are only expected to match when multicasting but not unicasting. It is also not a strict requirement for the multicast last octet to match.

#rnx

diff --git a/Engine/Plugins/VirtualProduction/LiveLinkOpenTrackIO/Source/LiveLinkOpenTrackIO/Private/LiveLinkOpenTrackIOSource.cpp b/Engine/Plugins/VirtualProduction/LiveLinkOpenTrackIO/Source/LiveLinkOpenTrackIO/Private/LiveLinkOpenTrackIOSource.cpp
--- a/Engine/Plugins/VirtualProduction/LiveLinkOpenTrackIO/Source/LiveLinkOpenTrackIO/Private/LiveLinkOpenTrackIOSource.cpp
+++ b/Engine/Plugins/VirtualProduction/LiveLinkOpenTrackIO/Source/LiveLinkOpenTrackIO/Private/LiveLinkOpenTrackIOSource.cpp
@@ -333,13 +333,19 @@ void FLiveLinkOpenTrackIOSource::PushDataToLiveLink_AnyThread(const FLiveLinkOpe
 	TRACE_CPUPROFILER_EVENT_SCOPE(OpenTrackIO::PushDataToLiveLink_AnyThread);
 
 	FLiveLinkOpenTrackIOCache* Cache;
-	if (OpenTrackIOCacheMap.Contains(InData.SourceId))
-	{
-		Cache = OpenTrackIOCacheMap.FindChecked(InData.SourceId).Get();
-	}
-	else
+
+	// Find the right cache for this OpenTrackIO source id and number
 	{
-		Cache = OpenTrackIOCacheMap.Add(InData.SourceId, MakePimpl<FLiveLinkOpenTrackIOCache>()).Get();
+		const FString OpenTrackIOStreamKey = FString::Printf(TEXT("%s:%d"), *InData.SourceId, InData.SourceNumber);
+
+		if (OpenTrackIOCacheMap.Contains(OpenTrackIOStreamKey))
+		{
+			Cache = OpenTrackIOCacheMap.FindChecked(OpenTrackIOStreamKey).Get();
+		}
+		else
+		{
+			Cache = OpenTrackIOCacheMap.Add(OpenTrackIOStreamKey, MakePimpl<FLiveLinkOpenTrackIOCache>()).Get();
+		}
 	}
 
 	if (!Cache->IsPacketInSequence(Header.SequenceNumber, InData.Timing.SampleRate.GetFrameRate()))
@@ -348,48 +354,53 @@ void FLiveLinkOpenTrackIOSource::PushDataToLiveLink_AnyThread(const FLiveLinkOpe
 		return;
 	}
 	
-	/** Check the static data camera fields on the inbound struct and see if we have new data to store / update. */
-	if (TOptional<FLiveLinkStaticDataStruct> StaticData = Cache->TryAddCamera(InData))
+	// Update the OpenTrackIO static camera data (make, model, etc.)
+	if (InData.Static.Camera.IsValid())
 	{
-		const FName SubjectName = Cache->GetSubjectNameFromData(InData);
+		Cache->StaticCamera = InData.Static.Camera;
+	}
 
-		Client->PushSubjectStaticData_AnyThread({ SourceGuid, SubjectName }, ULiveLinkLensRole::StaticClass(), MoveTemp(*StaticData));
-			  
-		// Remove the interim live link subject, if any.
-		if (!Cache->DefaultSubjectName.IsNone())
-		{
-			Client->RemoveSubject_AnyThread({ SourceGuid, Cache->DefaultSubjectName });
-			Cache->DefaultSubjectName = NAME_None;
-		}
+	// Update the OpenTrackIO static lens data
+	if (InData.Static.Lens.IsValid())
+	{
+		Cache->StaticLens = InData.Static.Lens;
 	}
-	else
+
+	// Detect if the new OpenTrackIO data changes the Subject Name
+	bool bSubjectNameChanged = false;
 	{
-		// If no static camera data has arrived, we add an interim live link subject named after the OpenTrackIO Source Number
-		// to support producers that do not send this information. If valid static camera data arrives later, it will be removed
-		// and replaced with the new one.
-		if (!Cache->StaticCameras.Num() && Cache->DefaultSubjectName.IsNone())
+		const FName SubjectName = Cache->GetSubjectNameFromData(InData);
+
+		if (Cache->SubjectName != SubjectName)
 		{
-			StaticData = Cache->MakeStaticData(InData);
-			check(StaticData.IsSet());
+			// Flag this because we'll need to push static data to add the new subject name.
+			bSubjectNameChanged = true;
 
-			// This is where the interim subject name is created and stored in the Cache's DefaultSubjectName.
-			Cache->DefaultSubjectName = FName(FString::Printf(TEXT("OpenTrackIO_%03u"), ConnectionSettings.SourceNumber));
+			// Remove the previous Live Link subject, since the name is different.
+			if (!Cache->SubjectName.IsNone())
+			{
+				Client->RemoveSubject_AnyThread({ SourceGuid, Cache->SubjectName });				
+			}
 
-			Client->PushSubjectStaticData_AnyThread({ SourceGuid, Cache->DefaultSubjectName }, ULiveLinkLensRole::StaticClass(), MoveTemp(*StaticData));
+			Cache->SubjectName = SubjectName;
 		}
 	}
 
-	TArray<FSubjectAndFrameDataTuple> FrameData = Cache->TryReadFrameData(InData);
-
-	/**
-	 * This loop will consume the data in FrameData and move it into the Live Link thread. You should consider the data invalid
-	 * after this loop completes.
-	 */
-	for (FSubjectAndFrameDataTuple& Item : FrameData)
+	// Update the static data if it has changed or if the subject name was updated.
 	{
-		Client->PushSubjectFrameData_AnyThread({ SourceGuid, Item.SubjectName }, MoveTemp(Item.FrameData));
+		FLiveLinkStaticDataStruct StaticData = Cache->MakeStaticData(InData);
+		const bool bStaticDataChanged = Cache->StaticData != StaticData;
+
+		if (bSubjectNameChanged || bStaticDataChanged)
+		{
+			Cache->StaticData.InitializeWith(StaticData);
+			Client->PushSubjectStaticData_AnyThread({ SourceGuid, Cache->SubjectName }, ULiveLinkLensRole::StaticClass(), MoveTemp(StaticData));
+		}
 	}
 
+	// Push the per-frame data
+	Client->PushSubjectFrameData_AnyThread({ SourceGuid, Cache->SubjectName }, Cache->MakeFrameData(InData));
+
 	Cache->UpdateLastKnownSequenceNumber(Header.SequenceNumber);
 }
 
diff --git a/Engine/Plugins/VirtualProduction/LiveLinkOpenTrackIO/Source/LiveLinkOpenTrackIO/Private/LiveLinkOpenTrackIOTranscoder.cpp b/Engine/Plugins/VirtualProduction/LiveLinkOpenTrackIO/Source/LiveLinkOpenTrackIO/Private/LiveLinkOpenTrackIOTranscoder.cpp
--- a/Engine/Plugins/VirtualProduction/LiveLinkOpenTrackIO/Source/LiveLinkOpenTrackIO/Private/LiveLinkOpenTrackIOTranscoder.cpp
+++ b/Engine/Plugins/VirtualProduction/LiveLinkOpenTrackIO/Source/LiveLinkOpenTrackIO/Private/LiveLinkOpenTrackIOTranscoder.cpp
@@ -8,24 +8,6 @@
 #include "Misc/QualifiedFrameTime.h"
 #include "Roles/LiveLinkCameraTypes.h"
 
-template<typename T>
-bool AddStaticData(FName SourceId, const T& Item, TMap<FName, T>& Container)
-{
-	if (!Item.IsValid())
-	{
-		return false;
-	}
-	
-	T& FoundItem = Container.FindOrAdd(SourceId);
-	if (FoundItem != Item)
-	{
-		FoundItem = Item;
-		return true;
-	}
-	FoundItem.UpdateStaticData(Item);
-
-	return false;
-}
 
 bool FLiveLinkOpenTrackIOCache::IsPacketInSequence(uint16 SequenceNumber, const FFrameRate& InRate) const
 {
@@ -50,28 +32,16 @@ bool FLiveLinkOpenTrackIOCache::IsPacketInSequence(uint16 SequenceNumber, const
 
 FName FLiveLinkOpenTrackIOCache::GetSubjectNameFromData(const FLiveLinkOpenTrackIOData& InData)
 {	
-	if (InData.Static.Camera.IsValid())
+	if (StaticCamera.IsSet())
 	{
-		return UE::OpenTrackIO::ConvertTypeToFName(InData.Static.Camera, FString::FromInt(InData.SourceNumber));
+		return UE::OpenTrackIO::ConvertTypeToFName(StaticCamera.GetValue(), FString::FromInt(InData.SourceNumber));
 	}
 
-	return DefaultSubjectName;
-}
-
-TOptional<FLiveLinkStaticDataStruct> FLiveLinkOpenTrackIOCache::TryAddCamera(const FLiveLinkOpenTrackIOData& Data)
-{
-	const FName SourceId = FName(Data.SourceId);
-
-	const bool bIsNewStatic =
-		   AddStaticData(SourceId, Data.Static.Lens, StaticLenses)
-		|| AddStaticData(SourceId, Data.Static.Camera, StaticCameras);
+	constexpr int32 NumCharsToTake = 6;
+	const int32 SourceIdLen = InData.SourceId.Len();
+	const FString SourceIdTail = InData.SourceId.Right(FMath::Min(NumCharsToTake, SourceIdLen));
 
-	if (!bIsNewStatic)
-	{
-		return {}; 
-	}
-
-	return MakeStaticData(Data);
+	return FName(FString::Printf(TEXT("OTrIO_%s_%03u"), *SourceIdTail, InData.SourceNumber));
 }
 
 FLiveLinkStaticDataStruct FLiveLinkOpenTrackIOCache::MakeStaticData(const FLiveLinkOpenTrackIOData& Data)
@@ -79,51 +49,61 @@ FLiveLinkStaticDataStruct FLiveLinkOpenTrackIOCache::MakeStaticData(const FLiveL
 	/** Lens class inherits from the camera data so we use that as our base for the static struct. */
 
 	FLiveLinkStaticDataStruct StaticDataStruct(FLiveLinkLensStaticData::StaticStruct());
-	FLiveLinkLensStaticData& LensData = *StaticDataStruct.Cast<FLiveLinkLensStaticData>();
+	FLiveLinkLensStaticData& NewStaticData = *StaticDataStruct.Cast<FLiveLinkLensStaticData>();
 
 	if (Data.Lens.PinholeFocalLength.IsSet())
 	{
-		LensData.bIsFocalLengthSupported = true;
+		NewStaticData.bIsFocalLengthSupported = true;
 	}
 
 	if (Data.Lens.FStop.IsSet())
 	{
-		LensData.bIsApertureSupported = true;
+		NewStaticData.bIsApertureSupported = true;
 	}
 
 	if (Data.Lens.FocusDistance.IsSet())
 	{
-		LensData.bIsFocusDistanceSupported = true;
+		NewStaticData.bIsFocusDistanceSupported = true;
 	}
 
-	if (Data.Static.Camera.ActiveSensorPhysicalDimensions.Height.IsSet()
-		&& Data.Static.Camera.ActiveSensorPhysicalDimensions.Width.IsSet())
+	if (Data.Transforms.Num() > 0 || Data.Lens.EntrancePupilOffset.IsSet())
 	{
-		LensData.bIsDynamicFilmbackSupported = true;
-
-		LensData.FilmBackHeight = Data.Static.Camera.ActiveSensorPhysicalDimensions.Height.GetValue();
-		LensData.FilmBackWidth = Data.Static.Camera.ActiveSensorPhysicalDimensions.Width.GetValue();
+		NewStaticData.bIsLocationSupported = true;
+		NewStaticData.bIsRotationSupported = true;
+		NewStaticData.bIsScaleSupported = true;
 	}
 
-	if (Data.Transforms.Num() > 0 || Data.Lens.EntrancePupilOffset.IsSet())
+	// For data that depends on OpenTrackIO static data, if this particular packet
+	// does not contain static data, then we check the cache.
 	{
-		LensData.bIsLocationSupported = true;
-		LensData.bIsRotationSupported = true;
-		LensData.bIsScaleSupported = true;
+		// In OpenTrackIO, we opt for sending the filmback in the frame data instead of static data,
+		// because some cameras with built-in undistortion may update this value dynamically.
+
+		if ((Data.Static.Camera.ActiveSensorPhysicalDimensions.Height.IsSet()
+			&& Data.Static.Camera.ActiveSensorPhysicalDimensions.Width.IsSet()))
+		{
+			NewStaticData.bIsDynamicFilmbackSupported = true;
+		} 
+		else if (StaticCamera.IsSet())
+		{
+			if (StaticCamera->ActiveSensorPhysicalDimensions.Height.IsSet()
+				&& StaticCamera->ActiveSensorPhysicalDimensions.Width.IsSet())
+			{
+				NewStaticData.bIsDynamicFilmbackSupported = true;
+			}
+		}
 	}
 
-	/** OpenTrackIO only allows for spherical models right now. */
-	LensData.LensModel = FName(TEXT("Spherical"));
+	// OpenTrackIO only allows for spherical models right now.
+	NewStaticData.LensModel = FName(TEXT("Spherical"));
 
 	return StaticDataStruct;
 }
 
 
-TArray<FSubjectAndFrameDataTuple> FLiveLinkOpenTrackIOCache::TryReadFrameData(const FLiveLinkOpenTrackIOData& Data)
+FLiveLinkFrameDataStruct FLiveLinkOpenTrackIOCache::MakeFrameData(const FLiveLinkOpenTrackIOData& Data)
 {
 	using namespace LiveLinkOpenTrackIOConversions;
-
-	TArray<FSubjectAndFrameDataTuple> Result;
 	
 	/** Compute the final transform by concatentating the transforms in the frame data. */
 	FTransform FinalTransform; 
@@ -160,9 +140,8 @@ TArray<FSubjectAndFrameDataTuple> FLiveLinkOpenTrackIOCache::TryReadFrameData(co
 		SampleTC = Data.Timing.Timecode.GetQualifiedFrameTime();
 	}
 
-	const FName SourceId = FName(Data.SourceId);
-	const FLiveLinkOpenTrackIOStaticCamera* Camera = StaticCameras.Find(SourceId);
-	const FLiveLinkOpenTrackIOStaticLens* Lens = StaticLenses.Find(SourceId);
+	const FLiveLinkOpenTrackIOStaticCamera* Camera = StaticCamera.IsSet() ? &StaticCamera.GetValue() : nullptr;
+	const FLiveLinkOpenTrackIOStaticLens* Lens = StaticLens.IsSet() ? &StaticLens.GetValue() : nullptr;
 		
 	FLiveLinkFrameDataStruct FrameDataStruct = FLiveLinkFrameDataStruct(FLiveLinkLensFrameData::StaticStruct());
 	FLiveLinkLensFrameData& CameraAndLensData = *FrameDataStruct.Cast<FLiveLinkLensFrameData>();
@@ -179,11 +158,5 @@ TArray<FSubjectAndFrameDataTuple> FLiveLinkOpenTrackIOCache::TryReadFrameData(co
 		CameraAndLensData.MetaData.StringMetaData.Add(FName(Field.Key), Field.Value);
 	}
 	
-	FName SubjectName = GetSubjectNameFromData(Data);
-	// Note that although we are only adding frame data for a single subject. This function returns an array to
-	// support sending frame data for other parts of the OpenTrackIO spec.
-	// 
-	Result.Add({SubjectName, MoveTemp(FrameDataStruct)});
-
-	return Result;
+	return FrameDataStruct;
 }
diff --git a/Engine/Plugins/VirtualProduction/LiveLinkOpenTrackIO/Source/LiveLinkOpenTrackIO/Private/LiveLinkOpenTrackIOTranscoder.h b/Engine/Plugins/VirtualProduction/LiveLinkOpenTrackIO/Source/LiveLinkOpenTrackIO/Private/LiveLinkOpenTrackIOTranscoder.h
--- a/Engine/Plugins/VirtualProduction/LiveLinkOpenTrackIO/Source/LiveLinkOpenTrackIO/Private/LiveLinkOpenTrackIOTranscoder.h
+++ b/Engine/Plugins/VirtualProduction/LiveLinkOpenTrackIO/Source/LiveLinkOpenTrackIO/Private/LiveLinkOpenTrackIOTranscoder.h
@@ -7,18 +7,22 @@
 #include "LiveLinkTypes.h"
 
 #include "Containers/Map.h"
+#include "LiveLinkLensTypes.h"
 #include "Misc/Optional.h"
 
-struct FSubjectAndFrameDataTuple
+struct FLiveLinkOpenTrackIOCache
 {
+	/** Last known Subject Name */
 	FName SubjectName;
-	FLiveLinkFrameDataStruct FrameData;
-};
 
-struct FLiveLinkOpenTrackIOCache
-{
-	TMap<FName, FLiveLinkOpenTrackIOStaticCamera> StaticCameras;
-	TMap<FName, FLiveLinkOpenTrackIOStaticLens>   StaticLenses;
+	/** Live Link Static Data, used to detect changes */
+	FLiveLinkStaticDataStruct StaticData;
+
+	/** OpenTrackIO static camera data (make, model, etc.) */
+	TOptional<FLiveLinkOpenTrackIOStaticCamera> StaticCamera;
+
+	/** OpenTrackIO static lens data */
+	TOptional<FLiveLinkOpenTrackIOStaticLens> StaticLens;
 
 	/**
 	 * Check the given sequence number and return true if it is the expected number. Sequence numbers should be incrementing
@@ -39,30 +43,14 @@ struct FLiveLinkOpenTrackIOCache
 		LastDataReceiveTimeInSeconds = FPlatformTime::Seconds();
 	}
 	
-	/**
-	 * Set sequence number 
-	 */
-
-	/** Subject name to use if it fails to make one from the data. */
-	FName DefaultSubjectName;
-
 	/** Constructs a subject name from the given data. */
 	FName GetSubjectNameFromData(const FLiveLinkOpenTrackIOData& Data);
 	
-	/**
-	 * Attempt to add this camera to our cache. If the camera does not exist or the is the same make / model but different values then the static
-	 * struct is returned.
-	 */
-	TOptional<FLiveLinkStaticDataStruct> TryAddCamera(const FLiveLinkOpenTrackIOData& StaticData);
-
 	/** Create Live Link static data from the OpenTrackIO data */
 	FLiveLinkStaticDataStruct MakeStaticData(const FLiveLinkOpenTrackIOData& StaticData);
 
-	/**
-	 * Attempt to read the per-frame data for the given OpenTrackIO sample. If valid data is present then a constructed frame data will be provided
-	 * for each subject in the cache.
-	 */
-	TArray<FSubjectAndFrameDataTuple> TryReadFrameData(const FLiveLinkOpenTrackIOData& Data);
+	/** Create Live Link the per-frame data from the OpenTrackIO data */
+	FLiveLinkFrameDataStruct MakeFrameData(const FLiveLinkOpenTrackIOData& Data);
 
 private:
 	/** Tracks the current sequence number for the source. */
