OpenTrackIO: Added a way to customize the subject names. Leaving empty or "Auto" will auto-generate based on the stream data.

Also made the selection of using multicast vs unicast explicit, and allow the user to pick a unicast port number.

#rnx

diff --git a/Engine/Plugins/VirtualProduction/LiveLinkOpenTrackIO/Source/LiveLinkOpenTrackIO/Private/LiveLinkOpenTrackIOConnectionSettings.h b/Engine/Plugins/VirtualProduction/LiveLinkOpenTrackIO/Source/LiveLinkOpenTrackIO/Private/LiveLinkOpenTrackIOConnectionSettings.h
--- a/Engine/Plugins/VirtualProduction/LiveLinkOpenTrackIO/Source/LiveLinkOpenTrackIO/Private/LiveLinkOpenTrackIOConnectionSettings.h
+++ b/Engine/Plugins/VirtualProduction/LiveLinkOpenTrackIO/Source/LiveLinkOpenTrackIO/Private/LiveLinkOpenTrackIOConnectionSettings.h
@@ -5,12 +5,50 @@
 
 #include "LiveLinkOpenTrackIOConnectionSettings.generated.h"
 
+
+// Pick between Multicast or Unicast for this connection
+UENUM(BlueprintType)
+enum class ELiveLinkOpenTrackIOProtocol : uint8
+{
+	Multicast UMETA(DisplayName = "Multicast"),
+	Unicast   UMETA(DisplayName = "Unicast")
+};
+
+
 USTRUCT()
 struct LIVELINKOPENTRACKIO_API FLiveLinkOpenTrackIOConnectionSettings
 {
 	GENERATED_BODY()
 
-	/** The Source Number is a user-configurable 8-bit value that determines the multicast IP address for a specific Source from a Producer. */
-	UPROPERTY(EditAnywhere, meta = (ClampMin = "1", ClampMax = "200", UIMin = "1", UIMax = "200"), Category = "Settings")
+	/** Using this is equivalent to leaving the SubjectName empty. */
+	static constexpr TCHAR AutoSubjectName[] = TEXT("Auto");
+
+	/** If empty or "Auto", one will be automatically generated. */
+	UPROPERTY(EditAnywhere, Category = "Settings")
+	FString SubjectName = AutoSubjectName;
+
+	/** Protocol selection: Multicast (default) or Unicast. */
+	UPROPERTY(EditAnywhere, Category = "Settings")
+	ELiveLinkOpenTrackIOProtocol Protocol = ELiveLinkOpenTrackIOProtocol::Multicast;
+
+	/**
+	 * The Source Number will be used as the last octet of the multicast address 235.135.1.[SourceNumber], per the opentrackio.org spec.
+	 * Use values in the range [1,200].
+	 */
+	UPROPERTY(EditAnywhere, Category = "Settings", meta = (
+		EditCondition = "Protocol == ELiveLinkOpenTrackIOProtocol::Multicast",
+		EditConditionHides,
+		ClampMin = "1", ClampMax = "200",
+		UIMin = "1", UIMax = "200"
+		))
 	uint8 SourceNumber = 1;
+
+	/** Unicast port number. 0 will self assign, so it is suggested to pick a port that can be used in your system. */
+	UPROPERTY(EditAnywhere, Category = "Settings", meta = (
+		EditCondition = "Protocol == ELiveLinkOpenTrackIOProtocol::Unicast",
+		EditConditionHides,
+		ClampMin = "1", ClampMax = "65535",
+		UIMin = "1", UIMax = "65535"
+		))
+	uint16 UnicastPort = 0;
 };
diff --git a/Engine/Plugins/VirtualProduction/LiveLinkOpenTrackIO/Source/LiveLinkOpenTrackIO/Private/LiveLinkOpenTrackIOSource.cpp b/Engine/Plugins/VirtualProduction/LiveLinkOpenTrackIO/Source/LiveLinkOpenTrackIO/Private/LiveLinkOpenTrackIOSource.cpp
--- a/Engine/Plugins/VirtualProduction/LiveLinkOpenTrackIO/Source/LiveLinkOpenTrackIO/Private/LiveLinkOpenTrackIOSource.cpp
+++ b/Engine/Plugins/VirtualProduction/LiveLinkOpenTrackIO/Source/LiveLinkOpenTrackIO/Private/LiveLinkOpenTrackIOSource.cpp
@@ -221,6 +221,10 @@ void FLiveLinkOpenTrackIOSource::ParseEndpointFromSourceSettings()
 	if (TOptional<FIPv4Endpoint> InUniEndpoint = ParseEndpoint(SavedSourceSettings->UnicastEndpoint))
 	{
 		UnicastEndpoint = *InUniEndpoint;
+
+		// Note: The Unicast port number is not used when using multicast, only the address (which is used as an interface address to bind to).
+		UnicastEndpoint.Port = ConnectionSettings.UnicastPort;
+
 		MulticastEndpoint = MakeMulticastAddressFromSourceNumber(ConnectionSettings.SourceNumber, SavedSourceSettings->MulticastPort);
 		ConnectionState = ELiveLinkOpenTrackIOState::EndpointsReady;
 	}
@@ -232,6 +236,9 @@ void FLiveLinkOpenTrackIOSource::InitializeSettings(ULiveLinkSourceSettings* Set
 
 	if (ULiveLinkOpenTrackIOSourceSettings* OpenTrackIOSettings = Cast<ULiveLinkOpenTrackIOSourceSettings>(Settings))
 	{
+		// Cache this for details to properties with the right context.
+		OpenTrackIOSettings->Protocol = ConnectionSettings.Protocol;
+
 		SavedSourceSettings = OpenTrackIOSettings;
 
 		ParseEndpointFromSourceSettings();
@@ -270,7 +277,7 @@ void FLiveLinkOpenTrackIOSource::OnSettingsChanged(ULiveLinkSourceSettings* Sett
 			if (PropertyName == GET_MEMBER_NAME_CHECKED(ULiveLinkOpenTrackIOSourceSettings, UnicastEndpoint)
 				|| PropertyName == GET_MEMBER_NAME_CHECKED(ULiveLinkOpenTrackIOSourceSettings, MulticastPort))
 			{
-				/** Start the reset process. */
+				// Start the reset process.
 				ResetSockets();
 			}
 		}
@@ -369,7 +376,7 @@ void FLiveLinkOpenTrackIOSource::PushDataToLiveLink_AnyThread(const FLiveLinkOpe
 	// Detect if the new OpenTrackIO data changes the Subject Name
 	bool bSubjectNameChanged = false;
 	{
-		const FName SubjectName = Cache->GetSubjectNameFromData(InData);
+		const FName SubjectName = Cache->GetSubjectNameFromData(ConnectionSettings.SubjectName, InData);
 
 		if (Cache->SubjectName != SubjectName)
 		{
@@ -406,6 +413,8 @@ void FLiveLinkOpenTrackIOSource::PushDataToLiveLink_AnyThread(const FLiveLinkOpe
 
 void FLiveLinkOpenTrackIOSource::CloseSockets()
 {
+	// Stop the receivers
+
 	if (MulticastReceiver)
 	{
 		MulticastReceiver->Stop();
@@ -416,9 +425,22 @@ void FLiveLinkOpenTrackIOSource::CloseSockets()
 		UnicastReceiver->Stop();
 	}
 
+	// Close the sockets
+
+	if (MulticastSocket)
+	{
+		MulticastSocket->Close();
+	}
+	if (UnicastSocket)
+	{
+		UnicastSocket->Close();
+	}
+
 	MulticastReceiver.Reset();
 	UnicastReceiver.Reset();
 
+	// Destroy the sockets
+
 	ISocketSubsystem* SocketSubsystem = ISocketSubsystem::Get(PLATFORM_SOCKETSUBSYSTEM);
 	check(SocketSubsystem);
 
@@ -436,43 +458,12 @@ void FLiveLinkOpenTrackIOSource::CloseSockets()
 	OpenTrackIOCacheMap.Reset();
 }
 
-bool FLiveLinkOpenTrackIOSource::OpenSockets()
+bool FLiveLinkOpenTrackIOSource::OpenMulticastSocket()
 {
-	if (ConnectionState != ELiveLinkOpenTrackIOState::EndpointsReady)
-	{
-		return false;
-	}
-
 	// Create an IPv4 TCP Socket
 	const uint32 ReceiveBufferSize = 2 * 1024 * 1024;
 
-	UnicastSocket = FUdpSocketBuilder(TEXT("UdpOpenTrackIO_UnicastSocket"))
-		.AsNonBlocking()
-		.BoundToEndpoint(UnicastEndpoint)
-		.WithMulticastLoopback()
-		.WithReceiveBufferSize(ReceiveBufferSize);
-
-	if (UnicastSocket == nullptr)
-	{
-		UE_LOG(LogLiveLinkOpenTrackIO, Error, TEXT("Failed to create OpenTrackIO socket on %s"), *UnicastEndpoint.ToString());
-		return false;
-	}
-	else
-	{
-		const int32 PortNo = UnicastSocket->GetPortNo();
-		FString IpPortStr = FString::Format(TEXT("{0}:{1}"), {UnicastEndpoint.Address.ToString(), PortNo});
-		UE_LOG(LogLiveLinkOpenTrackIO, Display, TEXT("OpenTrackIO unicast socket socket bound to '%s'."), *IpPortStr);
-
-		AsyncTask(ENamedThreads::GameThread, [IpPortStr, this]() {
-
-			if (IsValid(SavedSourceSettings))
-			{
-				SavedSourceSettings->AssignedUnicastEndpoint = IpPortStr;
-			}
-		});
-	}
-
-	/** Create a multicast socket.  This */
+	/** Create a multicast socket. */
 	MulticastSocket = FUdpSocketBuilder(TEXT("UdpOpenTrackIO_MulticastSocket"))
 		.AsNonBlocking()
 		.AsReusable()
@@ -482,7 +473,7 @@ bool FLiveLinkOpenTrackIOSource::OpenSockets()
 		 * If multiple bus instances bind the same unicast ip:port combination (allowed as the socket is marked as reusable), then for each packet
 		 * sent to that ip:port combination, only one of the instances (at the discretion of the OS) will receive it. The instance that receives the
 		 * packet may vary over time, seemingly based on the congestion of its socket. This isn't the intended usage.
-		 * 
+		 *
 		 * To allow traffic to be sent directly to unicast for discovery, set the interface and port for the unicast endpoint
 		 * However for legacy reason, keep binding this as well, although it might be unreliable in some cases
 		 */
@@ -495,33 +486,99 @@ bool FLiveLinkOpenTrackIOSource::OpenSockets()
 #endif
 		.WithReceiveBufferSize(ReceiveBufferSize);
 
-
-	if (MulticastSocket != nullptr)
+	if (!MulticastSocket)
 	{
-		UE::LiveLinkOpenTrackIO::Private::JoinedToGroup(UnicastEndpoint, MulticastEndpoint, MulticastSocket);
+		UE_LOG(LogLiveLinkOpenTrackIO, Warning, TEXT("StartTransport failed to create multicast socket on %s, joined to %s"), *UnicastEndpoint.ToString(), *MulticastEndpoint.ToString());
+		return false;
 	}
-	else
+
+	UE::LiveLinkOpenTrackIO::Private::JoinedToGroup(UnicastEndpoint, MulticastEndpoint, MulticastSocket);
+
+	// Create receiver for socket above
+
+	const FTimespan ThreadWaitTime = FTimespan::FromMilliseconds(100);
+
+	MulticastReceiver = MakeUnique<FUdpSocketReceiver>(MulticastSocket, ThreadWaitTime, TEXT("LiveLinkOpenTransportIO_MulticastReceiver"));
+	MulticastReceiver->OnDataReceived().BindRaw(this, &FLiveLinkOpenTrackIOSource::HandleInboundData);
+	MulticastReceiver->SetMaxReadBufferSize(2048);
+	MulticastReceiver->Start();
+
+	return true;
+}
+
+bool FLiveLinkOpenTrackIOSource::OpenUnicastSocket()
+{
+	// Create an IPv4 TCP Socket
+	const uint32 ReceiveBufferSize = 2 * 1024 * 1024;
+
+	UnicastSocket = FUdpSocketBuilder(TEXT("UdpOpenTrackIO_UnicastSocket"))
+		.AsNonBlocking()
+		.BoundToEndpoint(UnicastEndpoint)
+		.WithMulticastLoopback()
+		.WithReceiveBufferSize(ReceiveBufferSize);
+
+	if (!UnicastSocket)
 	{
-		UE_LOG(LogLiveLinkOpenTrackIO, Warning, TEXT("StartTransport failed to create multicast socket on %s, joined to %s"), *UnicastEndpoint.ToString(), *MulticastEndpoint.ToString());
+		UE_LOG(LogLiveLinkOpenTrackIO, Error, TEXT("Failed to create OpenTrackIO socket on %s"), *UnicastEndpoint.ToString());
 		return false;
 	}
 
-	/** Create receivers for sockets above */
+	const int32 PortNo = UnicastSocket->GetPortNo();
+	FString IpPortStr = FString::Format(TEXT("{0}:{1}"), { UnicastEndpoint.Address.ToString(), PortNo });
+	UE_LOG(LogLiveLinkOpenTrackIO, Display, TEXT("OpenTrackIO unicast socket socket bound to '%s'."), *IpPortStr);
+
+	AsyncTask(ENamedThreads::GameThread, [IpPortStr, this]() {
+
+		if (IsValid(SavedSourceSettings))
+		{
+			SavedSourceSettings->AssignedUnicastEndpoint = IpPortStr;
+		}
+		});
+
+	// Create receiver for socket above
+
 	const FTimespan ThreadWaitTime = FTimespan::FromMilliseconds(100);
-	if (MulticastSocket != nullptr)
+
+	UnicastReceiver = MakeUnique<FUdpSocketReceiver>(UnicastSocket, ThreadWaitTime, TEXT("LiveLinkOpenTransportIO_UnicastReceiver"));
+	UnicastReceiver->OnDataReceived().BindRaw(this, &FLiveLinkOpenTrackIOSource::HandleInboundData);
+	UnicastReceiver->SetMaxReadBufferSize(2048);
+	UnicastReceiver->Start();
+
+	return true;
+}
+
+
+bool FLiveLinkOpenTrackIOSource::OpenSockets()
+{
+	if (ConnectionState != ELiveLinkOpenTrackIOState::EndpointsReady)
 	{
-		MulticastReceiver = MakeUnique<FUdpSocketReceiver>(MulticastSocket, ThreadWaitTime, TEXT("LiveLinkOpenTransportIO_MulticastReceiver"));
-		MulticastReceiver->OnDataReceived().BindRaw(this, &FLiveLinkOpenTrackIOSource::HandleInboundData);
-		MulticastReceiver->SetMaxReadBufferSize(2048);
-		MulticastReceiver->Start();
+		return false;
 	}
 
-	if (UnicastSocket)
+	switch (ConnectionSettings.Protocol)
+	{
+	case ELiveLinkOpenTrackIOProtocol::Unicast:
 	{
-		UnicastReceiver = MakeUnique<FUdpSocketReceiver>(UnicastSocket, ThreadWaitTime, TEXT("LiveLinkOpenTransportIO_UnicastReceiver"));
-		UnicastReceiver->OnDataReceived().BindRaw(this, &FLiveLinkOpenTrackIOSource::HandleInboundData);
-		UnicastReceiver->SetMaxReadBufferSize(2048);
-		UnicastReceiver->Start();
+		if (!OpenUnicastSocket())
+		{
+			return false;
+		}
+
+		break;
+	}
+	case ELiveLinkOpenTrackIOProtocol::Multicast:
+	{
+		if (!OpenMulticastSocket())
+		{
+			return false;
+		}
+
+		break;
+	}
+	default:
+	{
+		return false;
+	}
 	}
 
 	ConnectionState = ELiveLinkOpenTrackIOState::Receiving;
diff --git a/Engine/Plugins/VirtualProduction/LiveLinkOpenTrackIO/Source/LiveLinkOpenTrackIO/Private/LiveLinkOpenTrackIOSource.h b/Engine/Plugins/VirtualProduction/LiveLinkOpenTrackIO/Source/LiveLinkOpenTrackIO/Private/LiveLinkOpenTrackIOSource.h
--- a/Engine/Plugins/VirtualProduction/LiveLinkOpenTrackIO/Source/LiveLinkOpenTrackIO/Private/LiveLinkOpenTrackIOSource.h
+++ b/Engine/Plugins/VirtualProduction/LiveLinkOpenTrackIO/Source/LiveLinkOpenTrackIO/Private/LiveLinkOpenTrackIOSource.h
@@ -91,6 +91,12 @@ private:
 	/** Create listening sockets for OpenTrack channels. Returns true if socket handles were successfully opened and receivers started. */
 	bool OpenSockets();
 
+	/** Create Multicast socket (called by OpenSockets) */
+	bool OpenMulticastSocket();
+
+	/** Create Uniicast socket (called by OpenSockets) */
+	bool OpenUnicastSocket();
+
 	/** Close any open sockets. */
 	void CloseSockets();
 
diff --git a/Engine/Plugins/VirtualProduction/LiveLinkOpenTrackIO/Source/LiveLinkOpenTrackIO/Private/LiveLinkOpenTrackIOTranscoder.cpp b/Engine/Plugins/VirtualProduction/LiveLinkOpenTrackIO/Source/LiveLinkOpenTrackIO/Private/LiveLinkOpenTrackIOTranscoder.cpp
--- a/Engine/Plugins/VirtualProduction/LiveLinkOpenTrackIO/Source/LiveLinkOpenTrackIO/Private/LiveLinkOpenTrackIOTranscoder.cpp
+++ b/Engine/Plugins/VirtualProduction/LiveLinkOpenTrackIO/Source/LiveLinkOpenTrackIO/Private/LiveLinkOpenTrackIOTranscoder.cpp
@@ -3,6 +3,7 @@
 #include "LiveLinkOpenTrackIOTranscoder.h"
 
 #include "LiveLinkLensTypes.h"
+#include "LiveLinkOpenTrackIOConnectionSettings.h"
 #include "LiveLinkOpenTrackIOConversions.h"
 #include "Misc/App.h"
 #include "Misc/QualifiedFrameTime.h"
@@ -30,18 +31,35 @@ bool FLiveLinkOpenTrackIOCache::IsPacketInSequence(uint16 SequenceNumber, const
 	return true;
 }
 
-FName FLiveLinkOpenTrackIOCache::GetSubjectNameFromData(const FLiveLinkOpenTrackIOData& InData)
+FName FLiveLinkOpenTrackIOCache::GetSubjectNameFromData(const FString& InSubjectName, const FLiveLinkOpenTrackIOData& InData)
 {	
+	// If no InSubjectName hint is provided, fallback names will be automatically generated based on the data.
+
+	if (!InSubjectName.IsEmpty() && (InSubjectName != FLiveLinkOpenTrackIOConnectionSettings::AutoSubjectName))
+	{
+		return FName(InSubjectName);
+	}
+
+	// First fallback: Static camera
 	if (StaticCamera.IsSet())
 	{
-		return UE::OpenTrackIO::ConvertTypeToFName(StaticCamera.GetValue(), FString::FromInt(InData.SourceNumber));
+		const FString Postfix = FString::Printf(TEXT("%03lld"), InData.SourceNumber);
+		return UE::OpenTrackIO::ConvertTypeToFName(StaticCamera.GetValue(), Postfix);
 	}
 
-	constexpr int32 NumCharsToTake = 6;
-	const int32 SourceIdLen = InData.SourceId.Len();
-	const FString SourceIdTail = InData.SourceId.Right(FMath::Min(NumCharsToTake, SourceIdLen));
+	// Next fallback: Static lens
+	if (StaticLens.IsSet())
+	{
+		const FString Postfix = FString::Printf(TEXT("%03lld"), InData.SourceNumber);
+		return UE::OpenTrackIO::ConvertTypeToFName(StaticLens.GetValue(), Postfix);
+	}
+
+	// Last fallback: Use the SourceId and Source Number of this OTrIO stream.
+
+	constexpr int32 NumLastSourceIdCharsToTake = 6;
+	const FString LastOfSourceId = InData.SourceId.Right(FMath::Min(NumLastSourceIdCharsToTake, InData.SourceId.Len()));
 
-	return FName(FString::Printf(TEXT("OTrIO_%s_%03u"), *SourceIdTail, InData.SourceNumber));
+	return FName(FString::Printf(TEXT("OTrIO_%s_%03lld"), *LastOfSourceId, InData.SourceNumber));
 }
 
 FLiveLinkStaticDataStruct FLiveLinkOpenTrackIOCache::MakeStaticData(const FLiveLinkOpenTrackIOData& Data)
diff --git a/Engine/Plugins/VirtualProduction/LiveLinkOpenTrackIO/Source/LiveLinkOpenTrackIO/Private/LiveLinkOpenTrackIOTranscoder.h b/Engine/Plugins/VirtualProduction/LiveLinkOpenTrackIO/Source/LiveLinkOpenTrackIO/Private/LiveLinkOpenTrackIOTranscoder.h
--- a/Engine/Plugins/VirtualProduction/LiveLinkOpenTrackIO/Source/LiveLinkOpenTrackIO/Private/LiveLinkOpenTrackIOTranscoder.h
+++ b/Engine/Plugins/VirtualProduction/LiveLinkOpenTrackIO/Source/LiveLinkOpenTrackIO/Private/LiveLinkOpenTrackIOTranscoder.h
@@ -44,7 +44,7 @@ struct FLiveLinkOpenTrackIOCache
 	}
 	
 	/** Constructs a subject name from the given data. */
-	FName GetSubjectNameFromData(const FLiveLinkOpenTrackIOData& Data);
+	FName GetSubjectNameFromData(const FString& Prefix, const FLiveLinkOpenTrackIOData& Data);
 	
 	/** Create Live Link static data from the OpenTrackIO data */
 	FLiveLinkStaticDataStruct MakeStaticData(const FLiveLinkOpenTrackIOData& StaticData);
diff --git a/Engine/Plugins/VirtualProduction/LiveLinkOpenTrackIO/Source/LiveLinkOpenTrackIO/Public/LiveLinkOpenTrackIOSourceSettings.h b/Engine/Plugins/VirtualProduction/LiveLinkOpenTrackIO/Source/LiveLinkOpenTrackIO/Public/LiveLinkOpenTrackIOSourceSettings.h
--- a/Engine/Plugins/VirtualProduction/LiveLinkOpenTrackIO/Source/LiveLinkOpenTrackIO/Public/LiveLinkOpenTrackIOSourceSettings.h
+++ b/Engine/Plugins/VirtualProduction/LiveLinkOpenTrackIO/Source/LiveLinkOpenTrackIO/Public/LiveLinkOpenTrackIOSourceSettings.h
@@ -3,6 +3,9 @@
 #pragma once
 
 #include "LiveLinkSourceSettings.h"
+
+#include "LiveLinkOpenTrackIOConnectionSettings.h"
+
 #include "LiveLinkOpenTrackIOSourceSettings.generated.h"
 
 UCLASS()
@@ -11,12 +14,20 @@ class LIVELINKOPENTRACKIO_API ULiveLinkOpenTrackIOSourceSettings : public ULiveL
 	GENERATED_BODY()
 
 public:
+
+	/** Protocol being used for the connection. */
+	UPROPERTY()
+	ELiveLinkOpenTrackIOProtocol Protocol = ELiveLinkOpenTrackIOProtocol::Multicast;
+
 	/**
 	 * The Multicast port to receive packets on.
 	 *
 	 * The default value 55555 per OpenTrackIO specifications.
 	 */
-	UPROPERTY(EditAnywhere, Category=Transport)
+	UPROPERTY(EditAnywhere, Category = Transport, meta = (
+		EditCondition = "Protocol == ELiveLinkOpenTrackIOProtocol::Multicast",
+		EditConditionHides
+	))
 	uint16 MulticastPort = 55555;
 
 	/**
@@ -40,6 +51,9 @@ public:
 	 * If no receiver is setup then this value will be an empty string. 
 	 *
 	 */
-	UPROPERTY(VisibleAnywhere, Category=Transport, meta = (DisplayName = "Assigned Endpoint") )
+	UPROPERTY(VisibleAnywhere, Category=Transport, meta = (
+		EditCondition = "Protocol == ELiveLinkOpenTrackIOProtocol::Unicast",
+		EditConditionHides
+		))
 	FString AssignedUnicastEndpoint = TEXT("");
 };
