OpenTrackIO: Use sender ip as SourceMachine.

This works both in multicast and unicast. The only potential problem is if different machines are sending to the same multicast group and port, but that situation should generally be avoided and use different source ids or port numbers instead.

#rnx

diff --git a/Engine/Plugins/VirtualProduction/LiveLinkOpenTrackIO/Source/LiveLinkOpenTrackIO/Private/LiveLinkOpenTrackIOSource.cpp b/Engine/Plugins/VirtualProduction/LiveLinkOpenTrackIO/Source/LiveLinkOpenTrackIO/Private/LiveLinkOpenTrackIOSource.cpp
--- a/Engine/Plugins/VirtualProduction/LiveLinkOpenTrackIO/Source/LiveLinkOpenTrackIO/Private/LiveLinkOpenTrackIOSource.cpp
+++ b/Engine/Plugins/VirtualProduction/LiveLinkOpenTrackIO/Source/LiveLinkOpenTrackIO/Private/LiveLinkOpenTrackIOSource.cpp
@@ -121,13 +121,8 @@ FLiveLinkOpenTrackIOSource::FLiveLinkOpenTrackIOSource(FLiveLinkOpenTrackIOConne
 	: ConnectionSettings(MoveTemp(InConnectionSettings))
 {
 	SourceType = LOCTEXT("SourceType_OpenTrackIO", "OpenTrack I/O");
-
-	SourceMachineName = FText::Format(
-		LOCTEXT("SourceNumber", "Source # {0}"),
-		ConnectionSettings.SourceNumber
-	);
-
 	SourceStatus = LOCTEXT("Startup", "Waiting for startup...");
+
 	DeferredStartDelegateHandle = FCoreDelegates::OnEndFrame.AddRaw(this, &FLiveLinkOpenTrackIOSource::Start);
 }
 
@@ -317,8 +312,19 @@ bool FLiveLinkOpenTrackIOSource::ConsumeInboundData()
 
 		TArrayView<uint8> MessageData(Message.Data->GetData(), Message.Data->Num());
 		TOptional<FOpenTrackIOHeaderWithData> ParsedData = UE::OpenTrackIO::Private::ParseOpenTrackIOPacket(MessageData);
+
 		if (ParsedData)
 		{
+			// Update Source Machine so the user can identify the sending machine.
+			if (LastSender_RunnableThread != Message.Sender)
+			{
+				LastSender_RunnableThread = Message.Sender;
+
+				AsyncTask(ENamedThreads::GameThread, [this, Sender = Message.Sender]() {
+						SourceMachineName = Sender.ToText();
+					});
+			}
+
 			bProcessedData = true;
 			OpenTrackIOPacketCount++;
 
diff --git a/Engine/Plugins/VirtualProduction/LiveLinkOpenTrackIO/Source/LiveLinkOpenTrackIO/Private/LiveLinkOpenTrackIOSource.h b/Engine/Plugins/VirtualProduction/LiveLinkOpenTrackIO/Source/LiveLinkOpenTrackIO/Private/LiveLinkOpenTrackIOSource.h
--- a/Engine/Plugins/VirtualProduction/LiveLinkOpenTrackIO/Source/LiveLinkOpenTrackIO/Private/LiveLinkOpenTrackIOSource.h
+++ b/Engine/Plugins/VirtualProduction/LiveLinkOpenTrackIO/Source/LiveLinkOpenTrackIO/Private/LiveLinkOpenTrackIOSource.h
@@ -198,4 +198,7 @@ private:
 
 	/** Counter to track the number of inbound OpenTrackIO packets received. */
 	std::atomic<uint64> OpenTrackIOPacketCount = 0;
+
+	/** Cached sender ip address, used to update the UI */
+	FIPv4Endpoint LastSender_RunnableThread;
 };
