OpenTrackIO: Omit port in SourceMachine.

The sender's port is not typically useful information.

#rnx

diff --git a/Engine/Plugins/VirtualProduction/LiveLinkOpenTrackIO/Source/LiveLinkOpenTrackIO/Private/LiveLinkOpenTrackIOSource.cpp b/Engine/Plugins/VirtualProduction/LiveLinkOpenTrackIO/Source/LiveLinkOpenTrackIO/Private/LiveLinkOpenTrackIOSource.cpp
--- a/Engine/Plugins/VirtualProduction/LiveLinkOpenTrackIO/Source/LiveLinkOpenTrackIO/Private/LiveLinkOpenTrackIOSource.cpp
+++ b/Engine/Plugins/VirtualProduction/LiveLinkOpenTrackIO/Source/LiveLinkOpenTrackIO/Private/LiveLinkOpenTrackIOSource.cpp
@@ -321,7 +321,7 @@ bool FLiveLinkOpenTrackIOSource::ConsumeInboundData()
 				LastSender_RunnableThread = Message.Sender;
 
 				AsyncTask(ENamedThreads::GameThread, [this, Sender = Message.Sender]() {
-						SourceMachineName = Sender.ToText();
+						SourceMachineName = Sender.Address.ToText();
 					});
 			}
 
