OpenTrackIO: Removed race condition when destroying the udp sockets.

#rnx

diff --git a/Engine/Plugins/VirtualProduction/LiveLinkOpenTrackIO/Source/LiveLinkOpenTrackIO/Private/LiveLinkOpenTrackIOSource.cpp b/Engine/Plugins/VirtualProduction/LiveLinkOpenTrackIO/Source/LiveLinkOpenTrackIO/Private/LiveLinkOpenTrackIOSource.cpp
--- a/Engine/Plugins/VirtualProduction/LiveLinkOpenTrackIO/Source/LiveLinkOpenTrackIO/Private/LiveLinkOpenTrackIOSource.cpp
+++ b/Engine/Plugins/VirtualProduction/LiveLinkOpenTrackIO/Source/LiveLinkOpenTrackIO/Private/LiveLinkOpenTrackIOSource.cpp
@@ -11,6 +11,7 @@
 #include "LiveLinkOpenTrackIOSourceSettings.h"
 #include "LiveLinkOpenTrackIOParser.h"
 #include "LiveLinkOpenTrackIOTranscoder.h"
+#include "LiveLinkOpenTrackIOUdpSocketReceiver.h"
 
 #include "Roles/LiveLinkCameraRole.h"
 #include "LiveLinkLensRole.h"
@@ -19,7 +20,6 @@
 #include "Misc/CoreDelegates.h"
 
 #include "Common/UdpSocketBuilder.h"
-#include "Common/UdpSocketReceiver.h"
 
 #include "Sockets.h"
 #include "SocketSubsystem.h"
@@ -431,15 +431,16 @@ void FLiveLinkOpenTrackIOSource::CloseSockets()
 		UnicastReceiver->Stop();
 	}
 
-	// Close the sockets
+	// Wait until they are confirmed to not be using the sockets.
 
-	if (MulticastSocket)
+	if (MulticastReceiver)
 	{
-		MulticastSocket->Close();
+		MulticastReceiver->WaitForStop();
 	}
-	if (UnicastSocket)
+
+	if (UnicastReceiver)
 	{
-		UnicastSocket->Close();
+		UnicastReceiver->WaitForStop();
 	}
 
 	MulticastReceiver.Reset();
@@ -504,7 +505,7 @@ bool FLiveLinkOpenTrackIOSource::OpenMulticastSocket()
 
 	const FTimespan ThreadWaitTime = FTimespan::FromMilliseconds(100);
 
-	MulticastReceiver = MakeUnique<FUdpSocketReceiver>(MulticastSocket, ThreadWaitTime, TEXT("LiveLinkOpenTransportIO_MulticastReceiver"));
+	MulticastReceiver = MakeUnique<FLiveLinkOpenTrackIOUdpSocketReceiver>(MulticastSocket, ThreadWaitTime, TEXT("LiveLinkOpenTransportIO_MulticastReceiver"));
 	MulticastReceiver->OnDataReceived().BindRaw(this, &FLiveLinkOpenTrackIOSource::HandleInboundData);
 	MulticastReceiver->SetMaxReadBufferSize(2048);
 	MulticastReceiver->Start();
@@ -545,7 +546,7 @@ bool FLiveLinkOpenTrackIOSource::OpenUnicastSocket()
 
 	const FTimespan ThreadWaitTime = FTimespan::FromMilliseconds(100);
 
-	UnicastReceiver = MakeUnique<FUdpSocketReceiver>(UnicastSocket, ThreadWaitTime, TEXT("LiveLinkOpenTransportIO_UnicastReceiver"));
+	UnicastReceiver = MakeUnique<FLiveLinkOpenTrackIOUdpSocketReceiver>(UnicastSocket, ThreadWaitTime, TEXT("LiveLinkOpenTransportIO_UnicastReceiver"));
 	UnicastReceiver->OnDataReceived().BindRaw(this, &FLiveLinkOpenTrackIOSource::HandleInboundData);
 	UnicastReceiver->SetMaxReadBufferSize(2048);
 	UnicastReceiver->Start();
diff --git a/Engine/Plugins/VirtualProduction/LiveLinkOpenTrackIO/Source/LiveLinkOpenTrackIO/Private/LiveLinkOpenTrackIOSource.h b/Engine/Plugins/VirtualProduction/LiveLinkOpenTrackIO/Source/LiveLinkOpenTrackIO/Private/LiveLinkOpenTrackIOSource.h
--- a/Engine/Plugins/VirtualProduction/LiveLinkOpenTrackIO/Source/LiveLinkOpenTrackIO/Private/LiveLinkOpenTrackIOSource.h
+++ b/Engine/Plugins/VirtualProduction/LiveLinkOpenTrackIO/Source/LiveLinkOpenTrackIO/Private/LiveLinkOpenTrackIOSource.h
@@ -8,6 +8,7 @@
 
 #include "LiveLinkOpenTrackIOConnectionSettings.h"
 #include "LiveLinkOpenTrackIOSourceSettings.h"
+#include "LiveLinkOpenTrackIOUdpSocketReceiver.h"
 
 #include "Misc/CoreMiscDefines.h"
 
@@ -17,8 +18,6 @@
 
 #include "Roles/LiveLinkTransformTypes.h"
 
-#include "Common/UdpSocketReceiver.h"
-
 #include "Delegates/IDelegateInstance.h"
 #include "MessageEndpoint.h"
 #include "IMessageContext.h"
@@ -158,10 +157,10 @@ private:
 	FSocket* MulticastSocket = nullptr;
 
 	/** Holds the unicast socket receiver. */
-	TUniquePtr<FUdpSocketReceiver> UnicastReceiver;
+	TUniquePtr<FLiveLinkOpenTrackIOUdpSocketReceiver> UnicastReceiver;
 	
 	/** Holds the multicast socket receiver. */
-	TUniquePtr<FUdpSocketReceiver> MulticastReceiver;
+	TUniquePtr<FLiveLinkOpenTrackIOUdpSocketReceiver> MulticastReceiver;
 
 	/** Holds the unicast socket. */
 	FSocket* UnicastSocket = nullptr;
diff --git a/Engine/Plugins/VirtualProduction/LiveLinkOpenTrackIO/Source/LiveLinkOpenTrackIO/Private/LiveLinkOpenTrackIOUdpSocketReceiver.h b/Engine/Plugins/VirtualProduction/LiveLinkOpenTrackIO/Source/LiveLinkOpenTrackIO/Private/LiveLinkOpenTrackIOUdpSocketReceiver.h
new file mode 100644
--- /dev/null
+++ b/Engine/Plugins/VirtualProduction/LiveLinkOpenTrackIO/Source/LiveLinkOpenTrackIO/Private/LiveLinkOpenTrackIOUdpSocketReceiver.h
@@ -0,0 +1,72 @@
+// Copyright Epic Games, Inc. All Rights Reserved.
+
+#pragma once
+
+#include "Common/UdpSocketReceiver.h"
+
+#include "HAL/PlatformProcess.h"
+
+#include <atomic>
+
+/**
+ * Same as FUdpSocketReceiver but with a function that blocks while inside Run(),
+ * which is useful when making sure it is safe to destroy the underlying sockets.
+ */
+class FLiveLinkOpenTrackIOUdpSocketReceiver
+	: public FUdpSocketReceiver
+{
+	using Super = FUdpSocketReceiver;
+
+public:
+
+	FLiveLinkOpenTrackIOUdpSocketReceiver(FSocket* InSocket, const FTimespan& InWaitTime, const TCHAR* InThreadName)
+		: Super(InSocket, InWaitTime, InThreadName)
+	{
+	}
+
+public:
+
+	//~ FRunnable interface Begin
+	virtual uint32 Run() override
+	{
+		bIsRunning = true;
+
+		uint32 RetVal = 0;
+
+		if (!bDoNotRunEverAgain)
+		{
+			RetVal = Super::Run();
+		}
+
+		bIsRunning = false;
+
+		return RetVal;
+	}
+	//~ FRunnable interface End
+
+	/** 
+	 * Blocks while inside Super::Run(). Call before destroying the sockets. 
+	 * Note: It won't be possible to ever Super::Run() after this call.
+	 */
+	void WaitForStop()
+	{
+		// We call stop just in case the user didn't, in which case we'd spin forever.
+		Stop();
+
+		// This flag will prevent entering Super::Run() even after WaitForStop() returns.
+		bDoNotRunEverAgain = true;
+
+		while (bIsRunning)
+		{
+			FPlatformProcess::SleepNoStats(2e-3); // 2 ms
+		}
+	}
+
+private:
+
+	/** Flag to determine if it is currently inside Super::Run() */
+	std::atomic<bool> bIsRunning = false;
+
+	/** Flag to avoid a posible race condition of WaitForStop called before Super::Run() enters */
+	std::atomic<bool> bDoNotRunEverAgain = false;
+};
