Fix python script name conflict with LiveLinkOpenTrackIOProtocol

#rnx

diff --git a/Engine/Plugins/VirtualProduction/LiveLinkOpenTrackIO/Source/LiveLinkOpenTrackIO/Private/LiveLinkOpenTrackIOConnectionSettings.h b/Engine/Plugins/VirtualProduction/LiveLinkOpenTrackIO/Source/LiveLinkOpenTrackIO/Private/LiveLinkOpenTrackIOConnectionSettings.h
--- a/Engine/Plugins/VirtualProduction/LiveLinkOpenTrackIO/Source/LiveLinkOpenTrackIO/Private/LiveLinkOpenTrackIOConnectionSettings.h
+++ b/Engine/Plugins/VirtualProduction/LiveLinkOpenTrackIO/Source/LiveLinkOpenTrackIO/Private/LiveLinkOpenTrackIOConnectionSettings.h
@@ -8,7 +8,7 @@
 
 // Pick between Multicast or Unicast for this connection
 UENUM(BlueprintType)
-enum class ELiveLinkOpenTrackIOProtocol : uint8
+enum class ELiveLinkOpenTrackIONetworkProtocol : uint8
 {
 	Multicast UMETA(DisplayName = "Multicast"),
 	Unicast   UMETA(DisplayName = "Unicast")
@@ -29,14 +29,14 @@ struct LIVELINKOPENTRACKIO_API FLiveLinkOpenTrackIOConnectionSettings
 
 	/** Protocol selection: Multicast (default) or Unicast. */
 	UPROPERTY(EditAnywhere, Category = "Settings")
-	ELiveLinkOpenTrackIOProtocol Protocol = ELiveLinkOpenTrackIOProtocol::Multicast;
+	ELiveLinkOpenTrackIONetworkProtocol Protocol = ELiveLinkOpenTrackIONetworkProtocol::Multicast;
 
 	/**
 	 * The Source Number will be used as the last octet of the multicast address 235.135.1.[SourceNumber], per the opentrackio.org spec.
 	 * Use values in the range [1,200].
 	 */
 	UPROPERTY(EditAnywhere, Category = "Settings", meta = (
-		EditCondition = "Protocol == ELiveLinkOpenTrackIOProtocol::Multicast",
+		EditCondition = "Protocol == ELiveLinkOpenTrackIONetworkProtocol::Multicast",
 		EditConditionHides,
 		ClampMin = "1", ClampMax = "200",
 		UIMin = "1", UIMax = "200"
@@ -45,7 +45,7 @@ struct LIVELINKOPENTRACKIO_API FLiveLinkOpenTrackIOConnectionSettings
 
 	/** Unicast port number. 0 will self assign, so it is suggested to pick a port that can be used in your system. */
 	UPROPERTY(EditAnywhere, Category = "Settings", meta = (
-		EditCondition = "Protocol == ELiveLinkOpenTrackIOProtocol::Unicast",
+		EditCondition = "Protocol == ELiveLinkOpenTrackIONetworkProtocol::Unicast",
 		EditConditionHides,
 		ClampMin = "1", ClampMax = "65535",
 		UIMin = "1", UIMax = "65535"
diff --git a/Engine/Plugins/VirtualProduction/LiveLinkOpenTrackIO/Source/LiveLinkOpenTrackIO/Private/LiveLinkOpenTrackIOSource.cpp b/Engine/Plugins/VirtualProduction/LiveLinkOpenTrackIO/Source/LiveLinkOpenTrackIO/Private/LiveLinkOpenTrackIOSource.cpp
--- a/Engine/Plugins/VirtualProduction/LiveLinkOpenTrackIO/Source/LiveLinkOpenTrackIO/Private/LiveLinkOpenTrackIOSource.cpp
+++ b/Engine/Plugins/VirtualProduction/LiveLinkOpenTrackIO/Source/LiveLinkOpenTrackIO/Private/LiveLinkOpenTrackIOSource.cpp
@@ -564,7 +564,7 @@ bool FLiveLinkOpenTrackIOSource::OpenSockets()
 
 	switch (ConnectionSettings.Protocol)
 	{
-	case ELiveLinkOpenTrackIOProtocol::Unicast:
+	case ELiveLinkOpenTrackIONetworkProtocol::Unicast:
 	{
 		if (!OpenUnicastSocket())
 		{
@@ -573,7 +573,7 @@ bool FLiveLinkOpenTrackIOSource::OpenSockets()
 
 		break;
 	}
-	case ELiveLinkOpenTrackIOProtocol::Multicast:
+	case ELiveLinkOpenTrackIONetworkProtocol::Multicast:
 	{
 		if (!OpenMulticastSocket())
 		{
diff --git a/Engine/Plugins/VirtualProduction/LiveLinkOpenTrackIO/Source/LiveLinkOpenTrackIO/Public/LiveLinkOpenTrackIOSourceSettings.h b/Engine/Plugins/VirtualProduction/LiveLinkOpenTrackIO/Source/LiveLinkOpenTrackIO/Public/LiveLinkOpenTrackIOSourceSettings.h
--- a/Engine/Plugins/VirtualProduction/LiveLinkOpenTrackIO/Source/LiveLinkOpenTrackIO/Public/LiveLinkOpenTrackIOSourceSettings.h
+++ b/Engine/Plugins/VirtualProduction/LiveLinkOpenTrackIO/Source/LiveLinkOpenTrackIO/Public/LiveLinkOpenTrackIOSourceSettings.h
@@ -17,7 +17,7 @@ public:
 
 	/** Protocol being used for the connection. */
 	UPROPERTY()
-	ELiveLinkOpenTrackIOProtocol Protocol = ELiveLinkOpenTrackIOProtocol::Multicast;
+	ELiveLinkOpenTrackIONetworkProtocol Protocol = ELiveLinkOpenTrackIONetworkProtocol::Multicast;
 
 	/**
 	 * The Multicast port to receive packets on.
@@ -25,7 +25,7 @@ public:
 	 * The default value 55555 per OpenTrackIO specifications.
 	 */
 	UPROPERTY(EditAnywhere, Category = Transport, meta = (
-		EditCondition = "Protocol == ELiveLinkOpenTrackIOProtocol::Multicast",
+		EditCondition = "Protocol == ELiveLinkOpenTrackIONetworkProtocol::Multicast",
 		EditConditionHides
 	))
 	uint16 MulticastPort = 55555;
@@ -52,7 +52,7 @@ public:
 	 *
 	 */
 	UPROPERTY(VisibleAnywhere, Category=Transport, meta = (
-		EditCondition = "Protocol == ELiveLinkOpenTrackIOProtocol::Unicast",
+		EditCondition = "Protocol == ELiveLinkOpenTrackIONetworkProtocol::Unicast",
 		EditConditionHides
 		))
 	FString AssignedUnicastEndpoint = TEXT("");
