Add packet segmentation support to OpenTrackIO plugin.

#rnx

diff --git a/Engine/Plugins/VirtualProduction/LiveLinkOpenTrackIO/Source/LiveLinkOpenTrackIO/Private/LiveLinkOpenTrackIODatagram.h b/Engine/Plugins/VirtualProduction/LiveLinkOpenTrackIO/Source/LiveLinkOpenTrackIO/Private/LiveLinkOpenTrackIODatagram.h
--- a/Engine/Plugins/VirtualProduction/LiveLinkOpenTrackIO/Source/LiveLinkOpenTrackIO/Private/LiveLinkOpenTrackIODatagram.h
+++ b/Engine/Plugins/VirtualProduction/LiveLinkOpenTrackIO/Source/LiveLinkOpenTrackIO/Private/LiveLinkOpenTrackIODatagram.h
@@ -68,6 +68,13 @@ struct FLiveLinkOpenTrackIODatagramHeader
 	UPROPERTY()
 	uint16 Checksum = 0;
 
+	/**
+	 * Returns true if the payload is complete and can be parsed.  If the payload is segmented then this will return false. 
+	 */
+	bool IsLastSegment() const
+	{
+		return LastSegmentFlagAndPayloadLength & 0x8000;
+	}
 
 	/** Helper function to return the size of the payload by and-ing the last 15 bits. */
 	uint16 GetPayloadSize() const
diff --git a/Engine/Plugins/VirtualProduction/LiveLinkOpenTrackIO/Source/LiveLinkOpenTrackIO/Private/LiveLinkOpenTrackIOParser.cpp b/Engine/Plugins/VirtualProduction/LiveLinkOpenTrackIO/Source/LiveLinkOpenTrackIO/Private/LiveLinkOpenTrackIOParser.cpp
--- a/Engine/Plugins/VirtualProduction/LiveLinkOpenTrackIO/Source/LiveLinkOpenTrackIO/Private/LiveLinkOpenTrackIOParser.cpp
+++ b/Engine/Plugins/VirtualProduction/LiveLinkOpenTrackIO/Source/LiveLinkOpenTrackIO/Private/LiveLinkOpenTrackIOParser.cpp
@@ -2,6 +2,7 @@
  
 #include "LiveLinkOpenTrackIOParser.h"
 
+#include "HAL/UnrealMemory.h"
 #include "LiveLinkOpenTrackIO.h"
 
 #include "LiveLinkOpenTrackIODatagram.h"
@@ -21,6 +22,74 @@
 #include "StructSerializer.h"
 #include "StructDeserializer.h"
 
+#include "HAL/IConsoleManager.h"
+
+// Payload size limiter to make sure senders don't endlessly send us data without a last segment. 
+static float PayloadSizeLimitInMb = 64;
+static FAutoConsoleVariableRef  CVarPayloadSizeLimitInMb(TEXT("OpenTrackIO.PayloadSizeLimit"), PayloadSizeLimitInMb, TEXT("Limits sizes of payloads that we will accept. The value is in megabytes."));
+
+void FLiveLinkOpenTrackIOPayload::AddBytes(uint32 InOffset, TArrayView<const uint8> InBytes)
+{
+	const float PayloadSizeLimitInFloatBytes = PayloadSizeLimitInMb * 1024 * 1024;
+	const uint32 PayloadLimitBytes = static_cast<uint32>(PayloadSizeLimitInFloatBytes);
+
+	if ((PayloadBytes.Num() + InOffset + InBytes.Num()) > PayloadLimitBytes)
+	{
+		bValidPayload = false;
+		UE_LOG(LogLiveLinkOpenTrackIO, Error, TEXT("Payload size limit exceeded %f MB. Use OpenTrackIO.PayloadSizeLimit to adjust this limit."), PayloadSizeLimitInMb);
+		return;
+	}
+	Sections.Add(FSegmentSection{ InOffset, InBytes.Num() });
+
+	// Preallocate and set the number to the new offset and size.
+	PayloadBytes.SetNum(InOffset + InBytes.Num());
+	
+	FMemory::Memcpy(PayloadBytes.GetData()+InOffset, InBytes.GetData(), InBytes.Num());
+}
+
+bool FLiveLinkOpenTrackIOPayload::IsComplete() const
+{
+	if (!bValidPayload)
+	{
+		return false;
+	}
+	
+	int32 TotalBytes = 0;
+
+	for (const FSegmentSection& Section : Sections)
+	{
+		TotalBytes += Section.Size;
+	}
+	
+	if (TotalBytes != PayloadBytes.Num())
+	{
+		// No possible way it is complete because TotalBytes should match our byte array size.
+		return false;
+	}
+
+	// If the size is equal then we need to make sure there are no overlaps between the received segments. If we did
+	// find an overlap then we can't consider it complete.
+	Algo::Sort(
+		Sections,
+		[](const FSegmentSection& A, const FSegmentSection& B)
+		{
+			return A.Offset < B.Offset;
+		}
+	);
+
+	bool bOverlaps = false;
+	// Check for overlaps
+	for (int32 Index = 1; !bOverlaps && Index < Sections.Num(); Index++)
+	{
+		if ((Sections[Index - 1].Offset + Sections[Index - 1].Size) < Sections[Index].Offset)
+		{
+			bOverlaps = true;
+			break;
+		}
+	}
+	return !bOverlaps;
+}
+
 namespace UE::OpenTrackIO::Private
 {
 
@@ -134,7 +203,7 @@ TOptional<FLiveLinkOpenTrackIOData> ParseJsonBlob(const FString& JsonBlob)
 	return OutData;
 }
 
-TOptional<FLiveLinkOpenTrackIOData> ParseCborBlob(TArrayView<uint8> InBytes)
+TOptional<FLiveLinkOpenTrackIOData> ParseCborBlob(TArrayView<const uint8> InBytes)
 {
 	TOptional<FLiveLinkOpenTrackIOData> OutData;
 	if (InBytes.IsEmpty())
@@ -162,7 +231,7 @@ TOptional<FLiveLinkOpenTrackIOData> ParseCborBlob(TArrayView<uint8> InBytes)
  * OpenTrackIO specifies using Fletcher-16 algorithm for checksum. 
  * Optimized version from: https://en.wikipedia.org/wiki/Fletcher%27s_checksum
  */
-uint16_t CalculateChecksum(TArrayView<uint8> HeaderBytes, TArrayView<uint8> PayloadBytes)
+uint16_t CalculateChecksum(TArrayView<const uint8> HeaderBytes, TArrayView<const uint8> PayloadBytes)
 {
 	uint32_t C0 = 0;
 	uint32_t C1 = 0;
@@ -246,18 +315,13 @@ inline FArchive& operator<<(FArchive& Ar, FLiveLinkOpenTrackIODatagramHeader& He
 	return Ar;
 }
 
-TOptional<FOpenTrackIOHeaderWithData> ParseOpenTrackIOPacket(TArrayView<uint8> Bytes)
+bool IsHeaderValid(TArrayView<const uint8> Bytes)
 {
-	TRACE_CPUPROFILER_EVENT_SCOPE(OpenTrackIO::ParseOpenTrackIOPacket);
-
 	const uint32 HeaderSize = sizeof(FLiveLinkOpenTrackIODatagramHeader);
-
-	TOptional<FOpenTrackIOHeaderWithData> OutData;
-
 	/** No way this is a valid datagram if we can't fit it in our header. */
 	if (Bytes.Num() < HeaderSize)
 	{
-		return OutData;
+		return false;
 	}
 
 	/** Check the header for validity first before trying to read the rest of the packet. */
@@ -273,13 +337,78 @@ TOptional<FOpenTrackIOHeaderWithData> ParseOpenTrackIOPacket(TArrayView<uint8> B
 				Identifier, OpenTrackIOHeaderId
 			);
 
-			return OutData;
+			return false;
+		}
+	}
+	return true;
+}
+
+TOptional<FLiveLinkOpenTrackIOData> ConvertPayloadToData(const FLiveLinkOpenTrackIODatagramHeader& Header, TArrayView<const uint8> Payload)
+{
+	/** Now attempt to parse the payload using either JSON or CBOR. */
+	TOptional<FLiveLinkOpenTrackIOData> Data;
+	if (Header.Encoding == JSON_Encoding)
+	{
+		FString JsonBlob;
+		FFileHelper::BufferToString(JsonBlob, Payload.GetData(), Payload.Num());
+		Data = ParseJsonBlob(JsonBlob);
+	}
+	else if (Header.Encoding == CBOR_Encoding)
+	{
+		Data = ParseCborBlob(Payload);
+	}
+	else
+	{
+		UE_LOG(LogLiveLinkOpenTrackIO, Error, TEXT("Unsupported encoding."));
+	}
+
+	return Data;
+}
+
+/**
+ * Helper function that takes header and payload data and converts it into a FLiveLinkOpenTrackIOData data struct.
+ *
+ * If this fails then the optional will be unset. 
+ */
+TOptional<FLiveLinkOpenTrackIOData> TryEmitOpenTrackData(FLiveLinkOpenTrackIODatagramHeader Header, TArrayView<const uint8> Payload)
+{
+	TOptional<FLiveLinkOpenTrackIOData> Data = ConvertPayloadToData(Header, Payload);
+	if (Data && Data->Protocol.IsSupported())
+	{
+		// Manually compensate for OpenTrackIO not explicitly stating the number of frames per frame.
+		{
+			// Calculate from sampleRate / timecodeRate, because it tells us how many samples there are per timecode frame.
+			const int32 RatioNum = Data->Timing.SampleRate.Num * Data->Timing.Timecode.FrameRate.Denom;
+			const int32 RatioDen = FMath::Max(1, Data->Timing.SampleRate.Denom * Data->Timing.Timecode.FrameRate.Num);
+			
+			// Enforce a minium of 1 SubframesPerFrame, which is equivalent to no frame subdivision.
+			Data->Timing.Timecode.SubframesPerFrame = FMath::Max(1, FMath::DivideAndRoundNearest(RatioNum, RatioDen));
 		}
 	}
+	else if (Data)
+	{
+		// We successfully parsed a data packet but we were unable to support the protocol.
+		//
+		Data.Reset();
+		UE_LOG(LogLiveLinkOpenTrackIO, Warning, TEXT("Unsupported OpenTrackIO protocol."));
+	}
+
+	return Data;
+}
+
+bool GetHeaderAndPayloadFromBytes(TArrayView<const uint8> Bytes, FOpenTrackIOHeaderWithPayload& OutPayloadContainer)
+{
+	TRACE_CPUPROFILER_EVENT_SCOPE(OpenTrackIO::GetHeaderAndPayloadFromBytes);
+
+	const uint32 HeaderSize = sizeof(FLiveLinkOpenTrackIODatagramHeader);
+	if (!IsHeaderValid(Bytes))
+	{
+		return false;		
+	}
 
 	/** Now read in the header to prepare parsing the payload. */
 	FMemoryReaderView Reader(Bytes);
-	FLiveLinkOpenTrackIODatagramHeader Header;
+	FLiveLinkOpenTrackIODatagramHeader& Header = OutPayloadContainer.GetMutableHeader();
 	Reader << Header;
 
 	/** Confirm that the payload size is within the number of bytes given to us. */
@@ -288,61 +417,38 @@ TOptional<FOpenTrackIOHeaderWithData> ParseOpenTrackIOPacket(TArrayView<uint8> B
 	{
 		/** The header does not match the inbound data. Abort! */
 		UE_LOG(LogLiveLinkOpenTrackIO, Error, TEXT("Payload size does not match provided OpenTrackIO packet header."));
-		return OutData;
+		return false;
 	}
 
-	TArrayView<uint8> HeaderMinusChecksum = Bytes.Slice(0, HeaderSize-2);
-	TArrayView<uint8> Payload = Bytes.Slice(HeaderSize, PayloadSize);
+	TArrayView<const uint8> HeaderMinusChecksum = Bytes.Slice(0, HeaderSize-2);
+	TArrayView<const uint8> Payload = Bytes.Slice(HeaderSize, PayloadSize);
 
 	/** Use a fletcher-16 algorithm to calculate the checksum of the data. */
 	uint16_t Checksum = CalculateChecksum(HeaderMinusChecksum, Payload);
 	if (Checksum != Header.Checksum)
 	{
 		UE_LOG(LogLiveLinkOpenTrackIO, Warning, TEXT("Failed to verify packet checksum."));
-		return OutData;
-	}
-
-	/** Now attempt to parse the payload using either JSON or CBOR. */
-	TOptional<FLiveLinkOpenTrackIOData> Data;
-	if (Header.Encoding == JSON_Encoding)
-	{
-		FString JsonBlob;
-		FFileHelper::BufferToString(JsonBlob, Payload.GetData(), Payload.Num());
-		Data = ParseJsonBlob(JsonBlob);
-	}
-	else if (Header.Encoding == CBOR_Encoding)
-	{
-		Data = ParseCborBlob(Payload);
-	}
-	else
-	{
-		UE_LOG(LogLiveLinkOpenTrackIO, Error, TEXT("Unsupported encoding."));
-		return OutData;
+		return false;
 	}
 
-	if (Data)
-	{
-		if (Data->Protocol.IsSupported())
-		{
-			// Manually compensate for OpenTrackIO not explicitly stating the number of frames per frame.
-			{
-				// Calculate from sampleRate / timecodeRate, because it tells us how many samples there are per timecode frame.
-				const int32 RatioNum = Data->Timing.SampleRate.Num * Data->Timing.Timecode.FrameRate.Denom;
-				const int32 RatioDen = FMath::Max(1, Data->Timing.SampleRate.Denom * Data->Timing.Timecode.FrameRate.Num);
+	OutPayloadContainer.GetMutablePayload().AddBytes(Header.SegmentOffset, Payload);
 
-				// Enforce a minium of 1 SubframesPerFrame, which is equivalent to no frame subdivision.
-				Data->Timing.Timecode.SubframesPerFrame = FMath::Max(1, FMath::DivideAndRoundNearest(RatioNum, RatioDen));
-			}
+	return true;
+}
 
-			OutData.Emplace(FOpenTrackIOHeaderWithData{Header, *Data});
-		}
-		else
-		{
-			UE_LOG(LogLiveLinkOpenTrackIO, Warning, TEXT("Unsupported OpenTrackIO protocol."));
-		}
+TOptional<FLiveLinkOpenTrackIOData> ParsePayload(const FOpenTrackIOHeaderWithPayload& HeaderAndPayload)
+{
+	const FLiveLinkOpenTrackIODatagramHeader& Header = HeaderAndPayload.GetHeader(); 
+	check(HeaderAndPayload.GetHeader().IsLastSegment());
+	
+	const FLiveLinkOpenTrackIOPayload& Payload = HeaderAndPayload.GetPayload();
+	if (!Payload.IsComplete())
+	{
+		UE_LOG(LogLiveLinkOpenTrackIO, Warning, TEXT("Discarding a partial payload because of missing segments."));
+		return {};
 	}
 
-	return OutData;
+	return TryEmitOpenTrackData(Header, Payload.GetBytes());
 }
 
 }
diff --git a/Engine/Plugins/VirtualProduction/LiveLinkOpenTrackIO/Source/LiveLinkOpenTrackIO/Private/LiveLinkOpenTrackIOParser.h b/Engine/Plugins/VirtualProduction/LiveLinkOpenTrackIO/Source/LiveLinkOpenTrackIO/Private/LiveLinkOpenTrackIOParser.h
--- a/Engine/Plugins/VirtualProduction/LiveLinkOpenTrackIO/Source/LiveLinkOpenTrackIO/Private/LiveLinkOpenTrackIOParser.h
+++ b/Engine/Plugins/VirtualProduction/LiveLinkOpenTrackIO/Source/LiveLinkOpenTrackIO/Private/LiveLinkOpenTrackIOParser.h
@@ -6,8 +6,108 @@
 #include "LiveLinkOpenTrackIODatagram.h"
 
 #include "Containers/ArrayView.h"
+#include "Containers/Array.h"
+
+/**
+ * Payload storage for inbound data stream.
+ *
+ * If the datagram header indicates a payload will be transmitted over multiple units, we store payload
+ * data by calling AddBytes per new packet arrival until all packets are received (last segment flag on the header).
+ *
+ * There is no reliable transmission protocol for OpenTrackIO. If we detect that the entire payload was not
+ * delivered by the time we receive the last segment flag then we throw out the entire payload. 
+ */
+struct FLiveLinkOpenTrackIOPayload
+{
+	FLiveLinkOpenTrackIOPayload() = default;
+
+	/** Add a partial payload to this class. */
+	void AddBytes(uint32 InOffset, TArrayView<const uint8> InBytes);
+
+	/**
+	 * Inspect all data received and confirm that it is complete.
+	 * 
+	 * To qualify as completed, we must have no gaps in our received data and the payload must not have exceeded
+	 * our payload limit set in OpenTrackIO.PayloadLimitSize
+	 */
+	bool IsComplete() const;
+
+	/**
+	 * Return a view into the data such that it can be passed into the CBOR or JSON parsers.
+	 *
+	 * You must ensure that you check that the payload is complete. 
+	 */
+	TArrayView<const uint8> GetBytes() const
+	{
+		return PayloadBytes;
+	}
+
+private:
+
+	/** Tracks if the payload is valid across all segments.  If we exceed our maximum payload size or we have received duplicate values */
+	bool bValidPayload = true;
+
+	/** Each segment received with Offset / Size pair. */
+	struct FSegmentSection
+	{
+		/** InOffset value */
+		uint32 Offset;
+
+		/** Size*/
+		int32 Size;
+	};
+
+	/**
+	 * Array to track all of the segment sections added to this payload.
+	 *
+	 * It is tagged mutable to support sorting within IsComplete() method. 
+	 */
+	mutable TArray<FSegmentSection> Sections;
+
+	/** Storage buffer for payload data from a packet stream. */
+	TArray<uint8> PayloadBytes;
+};
 
-using FOpenTrackIOHeaderWithData = TTuple<FLiveLinkOpenTrackIODatagramHeader, FLiveLinkOpenTrackIOData>;
+
+struct FOpenTrackIOHeaderWithPayload
+{
+	/**
+	 * Return a const reference to the header.
+	 */
+	const FLiveLinkOpenTrackIODatagramHeader& GetHeader() const
+	{
+		return Header;
+	}
+
+	/**
+	 * Return a const reference to the payload data. 
+	 */
+	const FLiveLinkOpenTrackIOPayload& GetPayload() const
+	{
+		return Payload;
+	}
+
+	/**
+	 * Read / write access to the header.
+	 */
+	FLiveLinkOpenTrackIODatagramHeader& GetMutableHeader()
+	{
+		return Header;
+	}
+
+	/**
+	 * Read / write access to the payload datga. 
+	 */
+	FLiveLinkOpenTrackIOPayload& GetMutablePayload()
+	{
+		return Payload;
+	}
+
+
+private:
+	FLiveLinkOpenTrackIODatagramHeader Header;
+	FLiveLinkOpenTrackIOPayload Payload;
+};
 
 namespace UE::OpenTrackIO::Private
 {
@@ -23,11 +123,20 @@ namespace UE::OpenTrackIO::Private
 	/**
 	 * Try to parse a CBOR blob. Return an optional FLiveLinkOpenTrackIOData type if CBOR parsing succeeds.
 	 */
-	TOptional<FLiveLinkOpenTrackIOData> ParseCborBlob(TArrayView<uint8> ByteArray);
+	TOptional<FLiveLinkOpenTrackIOData> ParseCborBlob(TArrayView<const uint8> ByteArray);
+
+	/**
+	 * From a byte data, split the byte stream into a header and payload part validating the header and running a Fletcher-16 checksum on the
+	 * header and payload.
+	 *
+	 * The result, if successful, is placed into OutPayloadContainer.  The function will return false if the checksum, header, or payload data
+	 * is invalid. 
+	 */
+	bool GetHeaderAndPayloadFromBytes(TArrayView<const uint8> Bytes, FOpenTrackIOHeaderWithPayload& OutPayloadContainer);
 
 	/**
-	 * Try to parse and read an OpenTrackIO packet.  If a packet read was successful then
-	 * a tuple will be returned that consists of the header for the datagram and the parsed payload.
+ 	 * Try to parse the payload from the given header and payload data. This will produce a FLiveLinkOpenTrackIOData struct that contains the full
+ 	 * values read from the payload data. If the parsers fail to read the payload then an unset optional will be returned.
 	 */
-	TOptional<FOpenTrackIOHeaderWithData> ParseOpenTrackIOPacket(TArrayView<uint8> Bytes);
+	TOptional<FLiveLinkOpenTrackIOData> ParsePayload(const FOpenTrackIOHeaderWithPayload& HeaderAndPayload);
 }
diff --git a/Engine/Plugins/VirtualProduction/LiveLinkOpenTrackIO/Source/LiveLinkOpenTrackIO/Private/LiveLinkOpenTrackIOSource.cpp b/Engine/Plugins/VirtualProduction/LiveLinkOpenTrackIO/Source/LiveLinkOpenTrackIO/Private/LiveLinkOpenTrackIOSource.cpp
--- a/Engine/Plugins/VirtualProduction/LiveLinkOpenTrackIO/Source/LiveLinkOpenTrackIO/Private/LiveLinkOpenTrackIOSource.cpp
+++ b/Engine/Plugins/VirtualProduction/LiveLinkOpenTrackIO/Source/LiveLinkOpenTrackIO/Private/LiveLinkOpenTrackIOSource.cpp
@@ -303,6 +303,7 @@ bool FLiveLinkOpenTrackIOSource::ConsumeInboundData()
 
 	bool bProcessedData = false;
 	FInboundMessage Message;
+
 	while (Inbound.Dequeue(Message))
 	{
 		if (!Message.Data.IsValid())
@@ -310,28 +311,37 @@ bool FLiveLinkOpenTrackIOSource::ConsumeInboundData()
 			continue;
 		}
 
-		TArrayView<uint8> MessageData(Message.Data->GetData(), Message.Data->Num());
-		TOptional<FOpenTrackIOHeaderWithData> ParsedData = UE::OpenTrackIO::Private::ParseOpenTrackIOPacket(MessageData);
+		TArrayView<const uint8> MessageData(Message.Data->GetData(), Message.Data->Num());
 
-		if (ParsedData)
+		FOpenTrackIOHeaderWithPayload& PayloadContainer = WorkingPayloads.FindOrAdd(Message.Sender);
+		if (UE::OpenTrackIO::Private::GetHeaderAndPayloadFromBytes(MessageData, PayloadContainer))
 		{
 			// Update Source Machine so the user can identify the sending machine.
 			if (LastSender_RunnableThread != Message.Sender)
 			{
 				LastSender_RunnableThread = Message.Sender;
-
 				AsyncTask(ENamedThreads::GameThread, [this, Sender = Message.Sender]() {
 						SourceMachineName = Sender.Address.ToText();
-					});
+				});
 			}
 
 			bProcessedData = true;
 			OpenTrackIOPacketCount++;
 
-			const FLiveLinkOpenTrackIODatagramHeader& Header = ParsedData->Get<0>();
-			FLiveLinkOpenTrackIOData& Data = ParsedData->Get<1>();
-
-			PushDataToLiveLink_AnyThread(Header, Data);
+			if (PayloadContainer.GetHeader().IsLastSegment())
+			{
+				{
+					// Scoped because want to prevent access to const & Header beyond this use.
+					const FLiveLinkOpenTrackIODatagramHeader& Header = PayloadContainer.GetHeader();
+					TOptional<FLiveLinkOpenTrackIOData> Data = UE::OpenTrackIO::Private::ParsePayload(PayloadContainer);
+					if (Data)
+					{
+						PushDataToLiveLink_AnyThread(Header, *Data);
+					}
+				}
+				// Remove the working payload from our table. This must happen 
+				WorkingPayloads.Remove(Message.Sender);
+			}
 		}
 		else
 		{
diff --git a/Engine/Plugins/VirtualProduction/LiveLinkOpenTrackIO/Source/LiveLinkOpenTrackIO/Private/LiveLinkOpenTrackIOSource.h b/Engine/Plugins/VirtualProduction/LiveLinkOpenTrackIO/Source/LiveLinkOpenTrackIO/Private/LiveLinkOpenTrackIOSource.h
--- a/Engine/Plugins/VirtualProduction/LiveLinkOpenTrackIO/Source/LiveLinkOpenTrackIO/Private/LiveLinkOpenTrackIOSource.h
+++ b/Engine/Plugins/VirtualProduction/LiveLinkOpenTrackIO/Source/LiveLinkOpenTrackIO/Private/LiveLinkOpenTrackIOSource.h
@@ -10,17 +10,13 @@
 #include "LiveLinkOpenTrackIOSourceSettings.h"
 #include "LiveLinkOpenTrackIOUdpSocketReceiver.h"
 
-#include "Misc/CoreMiscDefines.h"
+#include "LiveLinkOpenTrackIOParser.h"
 
 #include "Containers/MpscQueue.h"
 #include "Templates/UniquePtr.h"
 #include "Templates/PimplPtr.h"
 
-#include "Roles/LiveLinkTransformTypes.h"
-
 #include "Delegates/IDelegateInstance.h"
-#include "MessageEndpoint.h"
-#include "IMessageContext.h"
 #include "HAL/ThreadSafeBool.h"
 #include "HAL/Runnable.h"
 
@@ -130,6 +126,9 @@ private:
 	/** Pointer to implementation of our OpenTrackIO cache so that we can track camera and lens information. */
 	using FOpenTrackIOCachePtr = TPimplPtr<FLiveLinkOpenTrackIOCache>;
 
+	/** Keep track of payloads per endpoint. This is used to keep track of segmented data per endpoint.  */
+	TMap<FIPv4Endpoint, FOpenTrackIOHeaderWithPayload> WorkingPayloads;
+
 	/** Cache of data relevant to open track protocol. It is keyed by the sourceId. We can support multiple source ids */
 	TMap<FString, FOpenTrackIOCachePtr> OpenTrackIOCacheMap;
 	
diff --git a/Engine/Plugins/VirtualProduction/LiveLinkOpenTrackIO/Source/LiveLinkOpenTrackIO/Private/Tests/OpenTrackReaderTests.spec.cpp b/Engine/Plugins/VirtualProduction/LiveLinkOpenTrackIO/Source/LiveLinkOpenTrackIO/Private/Tests/OpenTrackReaderTests.spec.cpp
--- a/Engine/Plugins/VirtualProduction/LiveLinkOpenTrackIO/Source/LiveLinkOpenTrackIO/Private/Tests/OpenTrackReaderTests.spec.cpp
+++ b/Engine/Plugins/VirtualProduction/LiveLinkOpenTrackIO/Source/LiveLinkOpenTrackIO/Private/Tests/OpenTrackReaderTests.spec.cpp
@@ -121,16 +121,19 @@ void FOpenTrackIOReaderTests::Define()
 			uint64 Index = 0;
 			for (const FOTIO_DataValidation& Item : CborValidationData)
 			{
-				TArrayView<uint8> PacketView = AllPacketsView.Slice(Index, AllPacketsView.Num()-Index);
-				
-				TOptional<FOpenTrackIOHeaderWithData> ParsedPayload = UE::OpenTrackIO::Private::ParseOpenTrackIOPacket(PacketView);
-				FString ParsedPacket = TEXT("Cbor Parsed Packet No ") + FString::FromInt(Item.SequenceNumber); 
-				TestTrue(ParsedPacket, ParsedPayload.IsSet());
+				TArrayView<const uint8> PacketView = AllPacketsView.Slice(Index, AllPacketsView.Num()-Index);
+
+				FOpenTrackIOHeaderWithPayload PayloadContainer;
+				const bool bPayloadIsGood =
+					UE::OpenTrackIO::Private::GetHeaderAndPayloadFromBytes(PacketView, PayloadContainer);
+
+				FString ParsedPacket = TEXT("Cbor Parsed Packet No ") + FString::FromInt(Item.SequenceNumber);
+				TestTrue(ParsedPacket, bPayloadIsGood);
 
-				if (ParsedPayload)
+				if (TOptional<FLiveLinkOpenTrackIOData> ParsedPayload = UE::OpenTrackIO::Private::ParsePayload(PayloadContainer))
 				{
-					const FLiveLinkOpenTrackIODatagramHeader& Header = ParsedPayload->Get<0>();
-					const FLiveLinkOpenTrackIOData& Data = ParsedPayload->Get<1>();
+					const FLiveLinkOpenTrackIODatagramHeader& Header = PayloadContainer.GetHeader();
+					const FLiveLinkOpenTrackIOData& Data = *ParsedPayload;
 
 					TestTrue(ParsedPacket + " Has Payload", Header.GetPayloadSize() > 0);
 					TestTrue(ParsedPacket + " Matches Checksum", Item.Checksum == Header.Checksum);
@@ -155,16 +158,19 @@ void FOpenTrackIOReaderTests::Define()
 			uint64 Index = 0;
 			for (const FOTIO_DataValidation& Item : JsonValidationData)
 			{
-				TArrayView<uint8> PacketView = AllPacketsView.Slice(Index, AllPacketsView.Num()-Index);
-				
-				TOptional<FOpenTrackIOHeaderWithData> ParsedPayload = UE::OpenTrackIO::Private::ParseOpenTrackIOPacket(PacketView);
-				FString ParsedPacket = TEXT("JSON Parsed Packet No ") + FString::FromInt(Item.SequenceNumber); 
-				TestTrue(ParsedPacket, ParsedPayload.IsSet());
+				TArrayView<const uint8> PacketView = AllPacketsView.Slice(Index, AllPacketsView.Num()-Index);
+
+				FOpenTrackIOHeaderWithPayload PayloadContainer;
+				const bool bPayloadIsGood =
+					UE::OpenTrackIO::Private::GetHeaderAndPayloadFromBytes(PacketView, PayloadContainer);
+
+				FString ParsedPacket = TEXT("Json Parsed Packet No ") + FString::FromInt(Item.SequenceNumber);
+				TestTrue(ParsedPacket, bPayloadIsGood);
 
-				if (ParsedPayload)
+				if (TOptional<FLiveLinkOpenTrackIOData> ParsedPayload = UE::OpenTrackIO::Private::ParsePayload(PayloadContainer))
 				{
-					const FLiveLinkOpenTrackIODatagramHeader& Header = ParsedPayload->Get<0>();
-					const FLiveLinkOpenTrackIOData& Data = ParsedPayload->Get<1>();
+					const FLiveLinkOpenTrackIODatagramHeader& Header = PayloadContainer.GetHeader();
+					const FLiveLinkOpenTrackIOData& Data = *ParsedPayload;
 
 					TestTrue(ParsedPacket + " Has Payload", Header.GetPayloadSize() > 0);
 					TestTrue(ParsedPacket + " Matches Checksum", Item.Checksum == Header.Checksum);
