OpenTrackIO: Parse lens distortion data.

It will look for "Brown-Conrady U-D" and convert l1-l6 to OpenCV k1 k2 k3 p1 p2 using https://www.opentrackio.org/res/OpenCV_to_OpenTrackIO.pdf, which specifies powers of the focal length in mm for the conversion.

Other Lens Model names will be passed as is, in case there is a registered distortion shader than can handle them, and will pass the "custom" array of floats as parameters. opentrackio.org does not specify "custom" but right now we do not have an API to pass named parameters.

#rnx

diff --git a/Engine/Plugins/VirtualProduction/LiveLinkOpenTrackIO/Source/LiveLinkOpenTrackIO/Private/LiveLinkOpenTrackIOConversions.h b/Engine/Plugins/VirtualProduction/LiveLinkOpenTrackIO/Source/LiveLinkOpenTrackIO/Private/LiveLinkOpenTrackIOConversions.h
--- a/Engine/Plugins/VirtualProduction/LiveLinkOpenTrackIO/Source/LiveLinkOpenTrackIO/Private/LiveLinkOpenTrackIOConversions.h
+++ b/Engine/Plugins/VirtualProduction/LiveLinkOpenTrackIO/Source/LiveLinkOpenTrackIO/Private/LiveLinkOpenTrackIOConversions.h
@@ -39,6 +39,7 @@ namespace LiveLinkOpenTrackIOConversions
 		const FLiveLinkOpenTrackIOStaticCamera* InCamera
 	)
 	{
+		// FIZ
 		if (InLensData)
 		{
 			if (InLensData->FocusDistance.IsSet())
@@ -57,6 +58,7 @@ namespace LiveLinkOpenTrackIOConversions
 			}
 		}
 
+		// Filmback
 		if (InCamera)
 		{
 			if (InCamera->ActiveSensorPhysicalDimensions.Height.IsSet())
@@ -68,14 +70,91 @@ namespace LiveLinkOpenTrackIOConversions
 			{
 				OutUnrealLensData.FilmBackWidth = InCamera->ActiveSensorPhysicalDimensions.Width.GetValue();   // Both are in mm
 			}
+		}
 
-			if (InLensData)
-			{
-				// OpenTrackIO DistortionOffset is in mm, Unreal PrincipalPoint is normalized 0..1, centered at 0.5.
-				// We scale the input offset in mm by the filmback dimensions that are also in mm.
+		// Lens Distortion
+		//
+		// Only valid if the normalizing parameters are present (OpenTrackIO filmback and focal length, both in mm).
+		if (InCamera && InLensData 
+			&& InCamera->ActiveSensorPhysicalDimensions.Height.IsSet()
+			&& InCamera->ActiveSensorPhysicalDimensions.Width.IsSet()
+			&& InLensData->PinholeFocalLength.IsSet()
+			)
+		{
+			const float Width_mm = InCamera->ActiveSensorPhysicalDimensions.Width.GetValue();
+			const float Height_mm = InCamera->ActiveSensorPhysicalDimensions.Height.GetValue();
+			const float F_mm = InLensData->PinholeFocalLength.GetValue();
 
-				OutUnrealLensData.PrincipalPoint.X = 0.5 + InLensData->DistortionOffset.X * OutUnrealLensData.FilmBackWidth;
-				OutUnrealLensData.PrincipalPoint.Y = 0.5 + InLensData->DistortionOffset.Y * OutUnrealLensData.FilmBackHeight;
+			if ((Width_mm > KINDA_SMALL_NUMBER) && (Height_mm > KINDA_SMALL_NUMBER))  // Avoid division by zero.
+			{
+				// Cx Cy
+				{
+					// OpenTrackIO DistortionOffset is in mm, Unreal PrincipalPoint is normalized 0..1, centered at 0.5.
+					// We normalize DistortionOffset in mm by the filmback dimensions that are also in mm.
+
+					OutUnrealLensData.PrincipalPoint.X = 0.5 + InLensData->DistortionOffset.X / Width_mm;
+					OutUnrealLensData.PrincipalPoint.Y = 0.5 + InLensData->DistortionOffset.Y / Height_mm;
+				}
+
+				// FxFy
+				{
+					// OpenTrackIO specifies a singular focal length F. We normalize by the filmback dimensions to get
+					// focal length in UV units. That is, if they multiply by 3d coordinates normalized by depth,
+					// we get normalized screen coordiantes in the range [0,1].
+					OutUnrealLensData.FxFy.X = F_mm / Width_mm;
+					OutUnrealLensData.FxFy.Y = F_mm / Height_mm;
+				}
+
+				// Distortion parameters
+				{
+					// Note: The model name was already captured in the Live Link Static Data.
+					for (const FLiveLinkOpenTrackIOLens_DistortionCoeff& Distortion : InLensData->Distortion)
+					{
+						if (Distortion.Model == UE::OpenTrackIO::BrownConradyUD)
+						{
+							// OpenCV calibration outputs an Undistorted to Distorted equation, and Unreal does the same.
+							// 
+							// OpenTrackIO sends l1, l2, l3, l4, l5, l6, q1, q2.
+							//
+							// Unreal uses 5 OpenCV parameters (see SphericalLensModel.h)
+							// k1, k2, k3, p1, p2
+							//
+							// The conversions are: (see https://www.opentrackio.org/res/OpenCV_to_OpenTrackIO.pdf)
+							// 
+							// k1 = l1 * F^2
+							// k2 = l3 * F^4
+							// k3 = l5 * F^6
+							// 
+							// p1 = q1 * F^2
+							// p2 = q2 * F^2
+
+							if (Distortion.Radial.Num() == 6 && Distortion.Tangential.Num() == 2)
+							{
+								const float F_mm2 = F_mm  * F_mm;
+								const float F_mm4 = F_mm2 * F_mm2;
+								const float F_mm6 = F_mm2 * F_mm4;
+
+								OutUnrealLensData.DistortionParameters.SetNumUninitialized(5);
+
+								OutUnrealLensData.DistortionParameters[0] = Distortion.Radial[0] * F_mm2;
+								OutUnrealLensData.DistortionParameters[1] = Distortion.Radial[2] * F_mm4;
+								OutUnrealLensData.DistortionParameters[2] = Distortion.Radial[4] * F_mm6;
+
+								OutUnrealLensData.DistortionParameters[3] = Distortion.Tangential[0] * F_mm2;
+								OutUnrealLensData.DistortionParameters[4] = Distortion.Tangential[1] * F_mm2;
+							}
+						}
+						else
+						{
+							// Copy all parameters, we do not know what they are.
+							OutUnrealLensData.DistortionParameters = Distortion.Custom;
+						}
+
+						// For now we will pick the first available model, allowing custom distortion shaders
+						// to be used beyond spherical.
+						break;
+					}
+				}
 			}
 		}
 	}
diff --git a/Engine/Plugins/VirtualProduction/LiveLinkOpenTrackIO/Source/LiveLinkOpenTrackIO/Private/LiveLinkOpenTrackIOTranscoder.cpp b/Engine/Plugins/VirtualProduction/LiveLinkOpenTrackIO/Source/LiveLinkOpenTrackIO/Private/LiveLinkOpenTrackIOTranscoder.cpp
--- a/Engine/Plugins/VirtualProduction/LiveLinkOpenTrackIO/Source/LiveLinkOpenTrackIO/Private/LiveLinkOpenTrackIOTranscoder.cpp
+++ b/Engine/Plugins/VirtualProduction/LiveLinkOpenTrackIO/Source/LiveLinkOpenTrackIO/Private/LiveLinkOpenTrackIOTranscoder.cpp
@@ -112,8 +112,23 @@ FLiveLinkStaticDataStruct FLiveLinkOpenTrackIOCache::MakeStaticData(const FLiveL
 		}
 	}
 
-	// OpenTrackIO only allows for spherical models right now.
-	NewStaticData.LensModel = FName(TEXT("Spherical"));
+	// Lens Distortion Model
+
+	for (const FLiveLinkOpenTrackIOLens_DistortionCoeff& Distortion : Data.Lens.Distortion)
+	{
+		if (Distortion.Model == UE::OpenTrackIO::BrownConradyUD)
+		{
+			NewStaticData.LensModel = FName(TEXT("Spherical Lens Model"));
+		}
+		else
+		{
+			NewStaticData.LensModel = Distortion.Model;
+		}
+
+		// For now we will pick the first available model, allowing custom distortion shaders
+		// to be used instead of only Spherical.
+		break;
+	}
 
 	return StaticDataStruct;
 }
diff --git a/Engine/Plugins/VirtualProduction/LiveLinkOpenTrackIO/Source/LiveLinkOpenTrackIO/Private/LiveLinkOpenTrackIOTypes.h b/Engine/Plugins/VirtualProduction/LiveLinkOpenTrackIO/Source/LiveLinkOpenTrackIO/Private/LiveLinkOpenTrackIOTypes.h
--- a/Engine/Plugins/VirtualProduction/LiveLinkOpenTrackIO/Source/LiveLinkOpenTrackIO/Private/LiveLinkOpenTrackIOTypes.h
+++ b/Engine/Plugins/VirtualProduction/LiveLinkOpenTrackIO/Source/LiveLinkOpenTrackIO/Private/LiveLinkOpenTrackIOTypes.h
@@ -13,6 +13,20 @@
 #include "LiveLinkOpenTrackIOTypes.generated.h"
 
 
+
+namespace UE::OpenTrackIO
+{
+	/** Brown-Conrady D-U, the default for nameless distortion model per opentrackio.org */
+	static const FName BrownConradyDU = TEXT("Brown-Conrady D-U");
+
+	/** Brown-Conrady U-D which corresponds to "Spherical" model in Unreal */
+	static const FName BrownConradyUD = TEXT("Brown-Conrady U-D");
+
+	/** 3DE4 Anamorphic Standard Degree 4 */
+	static const FName _3DE4AnamorphicStandardDegree4 = TEXT("3DE4 Anamorphic Standard Degree 4");
+}
+
+
 /**
  * Represents frame rate type from OpenTrack I/O. Value is represented in hertz using an fractional representation
  * Numerator / Denominator. The default value is 24 / 1 or 24 cycles per second.
@@ -106,10 +120,10 @@ struct FLiveLinkOpenTrackIOSensorDimensions
 	bool operator!=(const FLiveLinkOpenTrackIOSensorDimensions&) const = default;
 	
 	UPROPERTY(EditAnywhere, BlueprintReadWrite, Category = "OpenTrackIO")
-	TOptional<double> Height;
+	TOptional<float> Height;
 
 	UPROPERTY(EditAnywhere, BlueprintReadWrite, Category = "OpenTrackIO")
-	TOptional<double> Width;
+	TOptional<float> Width;
 };
 
 /**
@@ -255,7 +269,7 @@ struct FLiveLinkOpenTrackIOStaticCamera
 	 * to the value of the parameter divided by 360 times the capture frame rate.
 	 */
 	UPROPERTY(meta = (ClampMin = "0.0", ClampMax = "360.0"))
-	double ShutterAngle = 180.0;
+	float ShutterAngle = 180.0;
 };
 
 USTRUCT(BlueprintType)
@@ -290,7 +304,7 @@ struct FLiveLinkOpenTrackIOStaticLens
 	 * of both projection-matrix-based and field-of-view-based rendering as per the OpenLensIO documentation.
 	 */
 	UPROPERTY(EditAnywhere, BlueprintReadWrite, Category = "OpenTrackIO")
-	double DistortionOverscanMax = 1.0;
+	float DistortionOverscanMax = 1.0;
 
 	/**
 	 * Static maximum overscan factor on lens undistortion. This is an alternative
@@ -299,7 +313,7 @@ struct FLiveLinkOpenTrackIOStaticLens
 	 * documentation.
 	 */
 	UPROPERTY(EditAnywhere, BlueprintReadWrite, Category = "OpenTrackIO")
-	double UndistortionOverscanMax = 1.0;
+	float UndistortionOverscanMax = 1.0;
 
 	/** Non-blank string naming lens manufacturer */
 	UPROPERTY(EditAnywhere, BlueprintReadWrite, Category = "OpenTrackIO")
@@ -319,7 +333,7 @@ struct FLiveLinkOpenTrackIOStaticLens
 
 	/** Nominal focal length of the lens. The number printed on the side of a prime lens, e.g. 50 mm, and undefined in the case of a zoom lens. */
 	UPROPERTY(EditAnywhere, BlueprintReadWrite, Category = "OpenTrackIO")
-	double NominalFocalLength = 50;
+	float NominalFocalLength = 50;
 
 	/** List of free strings that describe the history of calibrations of the lens. */
 	UPROPERTY(EditAnywhere, BlueprintReadWrite, Category = "OpenTrackIO")
@@ -806,10 +820,10 @@ struct FLiveLinkOpenTrackIOLens_DistortionOffset
 	}
 	
 	UPROPERTY(EditAnywhere, BlueprintReadWrite, Category = "OpenTrackIO")
-	double X = 0;
+	float X = 0;
 
 	UPROPERTY(EditAnywhere, BlueprintReadWrite, Category = "OpenTrackIO")
-	double Y = 0;
+	float Y = 0;
 };
 
 /**
@@ -840,13 +854,13 @@ struct FLiveLinkOpenTrackIOLens_Encoders
 	bool operator!=(const FLiveLinkOpenTrackIOLens_Encoders&) const = default;
 	
 	UPROPERTY(EditAnywhere, BlueprintReadWrite, Category = "OpenTrackIO")
-	double Focus = 0;
+	float Focus = 0;
 
 	UPROPERTY(EditAnywhere, BlueprintReadWrite, Category = "OpenTrackIO")
-	double Iris = 0;
+	float Iris = 0;
 
 	UPROPERTY(EditAnywhere, BlueprintReadWrite, Category = "OpenTrackIO")
-	double Zoom = 0;
+	float Zoom = 0;
 };
 
 
@@ -899,13 +913,13 @@ struct FLiveLinkOpenTrackIOLens_ExposureFalloff
 	bool operator!=(const FLiveLinkOpenTrackIOLens_ExposureFalloff&) const = default;
 	
 	UPROPERTY(EditAnywhere, BlueprintReadWrite, Category = "OpenTrackIO")
-	double A1 = 0;
+	float A1 = 0;
 
 	UPROPERTY(EditAnywhere, BlueprintReadWrite, Category = "OpenTrackIO")
-	double A2 = 0;
+	float A2 = 0;
 
 	UPROPERTY(EditAnywhere, BlueprintReadWrite, Category = "OpenTrackIO")
-	double A3 = 0;
+	float A3 = 0;
 };
 
 USTRUCT(BlueprintType)
@@ -924,13 +938,20 @@ struct FLiveLinkOpenTrackIOLens_DistortionCoeff
 	bool operator!=(const FLiveLinkOpenTrackIOLens_DistortionCoeff&) const = default;
 	
 	UPROPERTY(EditAnywhere, BlueprintReadWrite, Category = "OpenTrackIO")
-	TArray<double> Radial;
+	FName Model;
+
+	UPROPERTY(EditAnywhere, BlueprintReadWrite, Category = "OpenTrackIO")
+	TArray<float> Radial;
+
+	UPROPERTY(EditAnywhere, BlueprintReadWrite, Category = "OpenTrackIO")
+	TArray<float> Tangential;
 
 	UPROPERTY(EditAnywhere, BlueprintReadWrite, Category = "OpenTrackIO")
-	TArray<double> Tangential;
+	float Overscan = 0.0;
 
+	/** This is not currently in the spec but it would simplify sending any set of unnamed parameters */
 	UPROPERTY(EditAnywhere, BlueprintReadWrite, Category = "OpenTrackIO")
-	double Overscan = 0.0;
+	TArray<float> Custom;
 };
 
 USTRUCT(BlueprintType)
@@ -954,7 +975,7 @@ struct FLiveLinkOpenTrackIOLens
 	 * negotiation between a particular producer and consumer.
 	 */
 	UPROPERTY(EditAnywhere, BlueprintReadWrite, Category = "OpenTrackIO")
-	TArray<double> Custom;
+	TArray<float> Custom;
 
 	/**
 	 * A list of Distortion objects that each define the coefficients for calculating the distortion characteristics of a lens comprising radial
@@ -983,7 +1004,7 @@ struct FLiveLinkOpenTrackIOLens
 	 * virtual camera’s transform chain.
 	 */
 	UPROPERTY(EditAnywhere, BlueprintReadWrite, Category = "OpenTrackIO")
-	TOptional<double> EntrancePupilOffset;
+	TOptional<float> EntrancePupilOffset;
 
 	/**
 	 * Coefficients for calculating the exposure fall-off (vignetting) of a lens
@@ -993,15 +1014,15 @@ struct FLiveLinkOpenTrackIOLens
 
 	/** The linear f-number of the lens, equal to the focal length divided by the diameter of the entrance pupil. */
 	UPROPERTY(EditAnywhere, BlueprintReadWrite, Category = "OpenTrackIO")
-	TOptional<double> FStop;
+	TOptional<float> FStop;
 
 	/** Distance between the pinhole and the image plane in the simple CGI pinhole camera model. (millimeters) */
 	UPROPERTY(EditAnywhere, BlueprintReadWrite, Category = "OpenTrackIO")
-	TOptional<double> PinholeFocalLength;
+	TOptional<float> PinholeFocalLength;
 
 	/** Focus distance/position of the lens in meters. */ 
 	UPROPERTY(EditAnywhere, BlueprintReadWrite, Category = "OpenTrackIO")
-	TOptional<double> FocusDistance;
+	TOptional<float> FocusDistance;
 
 	/**  Offset in x and y of the centre of perspective projection of the virtual camera */
 	UPROPERTY(EditAnywhere, BlueprintReadWrite, Category = "OpenTrackIO")
@@ -1013,7 +1034,7 @@ struct FLiveLinkOpenTrackIOLens
 
 	/** T-stop: Linear t-number of the lens, equal to the F-number of the lens divided by the square root of the transmittance of the lens. */
 	UPROPERTY(EditAnywhere, BlueprintReadWrite, Category = "OpenTrackIO")
-	double TStop = 0.0;
+	float TStop = 0.0;
 };
 
 /**
