OpenTrackIO: Parse Brown Conrady DU, even though there isn't a shader for it.

#rnx

diff --git a/Engine/Plugins/VirtualProduction/LiveLinkOpenTrackIO/Source/LiveLinkOpenTrackIO/Private/LiveLinkOpenTrackIOConversions.h b/Engine/Plugins/VirtualProduction/LiveLinkOpenTrackIO/Source/LiveLinkOpenTrackIO/Private/LiveLinkOpenTrackIOConversions.h
--- a/Engine/Plugins/VirtualProduction/LiveLinkOpenTrackIO/Source/LiveLinkOpenTrackIO/Private/LiveLinkOpenTrackIOConversions.h
+++ b/Engine/Plugins/VirtualProduction/LiveLinkOpenTrackIO/Source/LiveLinkOpenTrackIO/Private/LiveLinkOpenTrackIOConversions.h
@@ -147,6 +147,8 @@ namespace LiveLinkOpenTrackIOConversions
 						else
 						{
 							// Copy all parameters, we do not know what they are.
+							OutUnrealLensData.DistortionParameters.Append(Distortion.Radial.GetData(), Distortion.Radial.Num());
+							OutUnrealLensData.DistortionParameters.Append(Distortion.Tangential.GetData(), Distortion.Tangential.Num());
 							OutUnrealLensData.DistortionParameters = Distortion.Custom;
 						}
 
