Add additional logging to OpenTrackIO Parser to help identify failure cases.

#rnx

diff --git a/Engine/Plugins/VirtualProduction/LiveLinkOpenTrackIO/Source/LiveLinkOpenTrackIO/Private/LiveLinkOpenTrackIOParser.cpp b/Engine/Plugins/VirtualProduction/LiveLinkOpenTrackIO/Source/LiveLinkOpenTrackIO/Private/LiveLinkOpenTrackIOParser.cpp
--- a/Engine/Plugins/VirtualProduction/LiveLinkOpenTrackIO/Source/LiveLinkOpenTrackIO/Private/LiveLinkOpenTrackIOParser.cpp
+++ b/Engine/Plugins/VirtualProduction/LiveLinkOpenTrackIO/Source/LiveLinkOpenTrackIO/Private/LiveLinkOpenTrackIOParser.cpp
@@ -403,6 +403,7 @@ bool GetHeaderAndPayloadFromBytes(TArrayView<const uint8> Bytes, FOpenTrackIOHea
 	const uint32 HeaderSize = sizeof(FLiveLinkOpenTrackIODatagramHeader);
 	if (!IsHeaderValid(Bytes))
 	{
+		UE_LOG(LogLiveLinkOpenTrackIO, Error, TEXT("Invalid OpenTrackIO Header Received."));
 		return false;		
 	}
 
