OpenTrackIO: Removed conversion from B-C distortion to Spherical.

Now we pass the unmodified parameters and let the corresponding lens model interpret them.

#rnx

diff --git a/Engine/Plugins/VirtualProduction/LiveLinkOpenTrackIO/Source/LiveLinkOpenTrackIO/Private/LiveLinkOpenTrackIOConversions.h b/Engine/Plugins/VirtualProduction/LiveLinkOpenTrackIO/Source/LiveLinkOpenTrackIO/Private/LiveLinkOpenTrackIOConversions.h
--- a/Engine/Plugins/VirtualProduction/LiveLinkOpenTrackIO/Source/LiveLinkOpenTrackIO/Private/LiveLinkOpenTrackIOConversions.h
+++ b/Engine/Plugins/VirtualProduction/LiveLinkOpenTrackIO/Source/LiveLinkOpenTrackIO/Private/LiveLinkOpenTrackIOConversions.h
@@ -106,56 +106,16 @@ namespace LiveLinkOpenTrackIOConversions
 				}
 
 				// Distortion parameters
+				// Note: The model name was already captured in the Live Link Static Data.
+				for (const FLiveLinkOpenTrackIOLens_DistortionCoeff& Distortion : InLensData->Distortion)
 				{
-					// Note: The model name was already captured in the Live Link Static Data.
-					for (const FLiveLinkOpenTrackIOLens_DistortionCoeff& Distortion : InLensData->Distortion)
-					{
-						if (Distortion.Model == UE::OpenTrackIO::BrownConradyUD)
-						{
-							// OpenCV calibration outputs an Undistorted to Distorted equation, and Unreal does the same.
-							// 
-							// OpenTrackIO sends l1, l2, l3, l4, l5, l6, q1, q2.
-							//
-							// Unreal uses 5 OpenCV parameters (see SphericalLensModel.h)
-							// k1, k2, k3, p1, p2
-							//
-							// The conversions are: (see https://www.opentrackio.org/res/OpenCV_to_OpenTrackIO.pdf)
-							// 
-							// k1 = l1 * F^2
-							// k2 = l3 * F^4
-							// k3 = l5 * F^6
-							// 
-							// p1 = q1 * F^2
-							// p2 = q2 * F^2
-
-							if (Distortion.Radial.Num() == 6 && Distortion.Tangential.Num() == 2)
-							{
-								const float F_mm2 = F_mm  * F_mm;
-								const float F_mm4 = F_mm2 * F_mm2;
-								const float F_mm6 = F_mm2 * F_mm4;
-
-								OutUnrealLensData.DistortionParameters.SetNumUninitialized(5);
-
-								OutUnrealLensData.DistortionParameters[0] = Distortion.Radial[0] * F_mm2;
-								OutUnrealLensData.DistortionParameters[1] = Distortion.Radial[2] * F_mm4;
-								OutUnrealLensData.DistortionParameters[2] = Distortion.Radial[4] * F_mm6;
-
-								OutUnrealLensData.DistortionParameters[3] = Distortion.Tangential[0] * F_mm2;
-								OutUnrealLensData.DistortionParameters[4] = Distortion.Tangential[1] * F_mm2;
-							}
-						}
-						else
-						{
-							// Copy all parameters, we do not know what they are.
-							OutUnrealLensData.DistortionParameters.Append(Distortion.Radial.GetData(), Distortion.Radial.Num());
-							OutUnrealLensData.DistortionParameters.Append(Distortion.Tangential.GetData(), Distortion.Tangential.Num());
-							OutUnrealLensData.DistortionParameters = Distortion.Custom;
-						}
-
-						// For now we will pick the first available model, allowing custom distortion shaders
-						// to be used beyond spherical.
-						break;
-					}
+					// Copy all parameters, we do not know what they are.
+					OutUnrealLensData.DistortionParameters.Append(Distortion.Radial.GetData(), Distortion.Radial.Num());
+					OutUnrealLensData.DistortionParameters.Append(Distortion.Tangential.GetData(), Distortion.Tangential.Num());
+					OutUnrealLensData.DistortionParameters = Distortion.Custom;
+
+					// For now we will pick the first model in the array.
+					break;
 				}
 			}
 		}
diff --git a/Engine/Plugins/VirtualProduction/LiveLinkOpenTrackIO/Source/LiveLinkOpenTrackIO/Private/LiveLinkOpenTrackIOTranscoder.cpp b/Engine/Plugins/VirtualProduction/LiveLinkOpenTrackIO/Source/LiveLinkOpenTrackIO/Private/LiveLinkOpenTrackIOTranscoder.cpp
--- a/Engine/Plugins/VirtualProduction/LiveLinkOpenTrackIO/Source/LiveLinkOpenTrackIO/Private/LiveLinkOpenTrackIOTranscoder.cpp
+++ b/Engine/Plugins/VirtualProduction/LiveLinkOpenTrackIO/Source/LiveLinkOpenTrackIO/Private/LiveLinkOpenTrackIOTranscoder.cpp
@@ -116,17 +116,17 @@ FLiveLinkStaticDataStruct FLiveLinkOpenTrackIOCache::MakeStaticData(const FLiveL
 
 	for (const FLiveLinkOpenTrackIOLens_DistortionCoeff& Distortion : Data.Lens.Distortion)
 	{
-		if (Distortion.Model == UE::OpenTrackIO::BrownConradyUD)
+		if (Distortion.Model.IsNone())
 		{
-			NewStaticData.LensModel = FName(TEXT("Spherical Lens Model"));
+			// OpentrackIO's default model is "Brown-Conrady D-U".
+			NewStaticData.LensModel = UE::OpenTrackIO::BrownConradyDU;
 		}
 		else
 		{
 			NewStaticData.LensModel = Distortion.Model;
 		}
 
-		// For now we will pick the first available model, allowing custom distortion shaders
-		// to be used instead of only Spherical.
+		// For now we will pick the first model in the array.
 		break;
 	}
 
diff --git a/Engine/Plugins/VirtualProduction/LiveLinkOpenTrackIO/Source/LiveLinkOpenTrackIO/Private/LiveLinkOpenTrackIOTypes.h b/Engine/Plugins/VirtualProduction/LiveLinkOpenTrackIO/Source/LiveLinkOpenTrackIO/Private/LiveLinkOpenTrackIOTypes.h
--- a/Engine/Plugins/VirtualProduction/LiveLinkOpenTrackIO/Source/LiveLinkOpenTrackIO/Private/LiveLinkOpenTrackIOTypes.h
+++ b/Engine/Plugins/VirtualProduction/LiveLinkOpenTrackIO/Source/LiveLinkOpenTrackIO/Private/LiveLinkOpenTrackIOTypes.h
@@ -18,12 +18,6 @@ namespace UE::OpenTrackIO
 {
 	/** Brown-Conrady D-U, the default for nameless distortion model per opentrackio.org */
 	static const FName BrownConradyDU = TEXT("Brown-Conrady D-U");
-
-	/** Brown-Conrady U-D which corresponds to "Spherical" model in Unreal */
-	static const FName BrownConradyUD = TEXT("Brown-Conrady U-D");
-
-	/** 3DE4 Anamorphic Standard Degree 4 */
-	static const FName _3DE4AnamorphicStandardDegree4 = TEXT("3DE4 Anamorphic Standard Degree 4");
 }
 
 
