OpenTrackIO: Fix overwriting radial and tangential with custom.

#rnx

diff --git a/Engine/Plugins/VirtualProduction/LiveLinkOpenTrackIO/Source/LiveLinkOpenTrackIO/Private/LiveLinkOpenTrackIOConversions.h b/Engine/Plugins/VirtualProduction/LiveLinkOpenTrackIO/Source/LiveLinkOpenTrackIO/Private/LiveLinkOpenTrackIOConversions.h
--- a/Engine/Plugins/VirtualProduction/LiveLinkOpenTrackIO/Source/LiveLinkOpenTrackIO/Private/LiveLinkOpenTrackIOConversions.h
+++ b/Engine/Plugins/VirtualProduction/LiveLinkOpenTrackIO/Source/LiveLinkOpenTrackIO/Private/LiveLinkOpenTrackIOConversions.h
@@ -112,7 +112,7 @@ namespace LiveLinkOpenTrackIOConversions
 					// Copy all parameters, we do not know what they are.
 					OutUnrealLensData.DistortionParameters.Append(Distortion.Radial.GetData(), Distortion.Radial.Num());
 					OutUnrealLensData.DistortionParameters.Append(Distortion.Tangential.GetData(), Distortion.Tangential.Num());
-					OutUnrealLensData.DistortionParameters = Distortion.Custom;
+					OutUnrealLensData.DistortionParameters.Append(Distortion.Custom.GetData(), Distortion.Custom.Num());
 
 					// For now we will pick the first model in the array.
 					break;
