Placeholder LiveLinkRole and LiveLinkTypes for the OpenTrackIO data.

#rnx

diff --git a/Engine/Plugins/VirtualProduction/LiveLinkOpenTrackIO/Source/LiveLinkOpenTrackIO/Private/LiveLinkOpenTrackIOLiveLinkTypes.h b/Engine/Plugins/VirtualProduction/LiveLinkOpenTrackIO/Source/LiveLinkOpenTrackIO/Private/LiveLinkOpenTrackIOLiveLinkTypes.h
new file mode 100644
--- /dev/null
+++ b/Engine/Plugins/VirtualProduction/LiveLinkOpenTrackIO/Source/LiveLinkOpenTrackIO/Private/LiveLinkOpenTrackIOLiveLinkTypes.h
@@ -0,0 +1,42 @@
+// Copyright Epic Games, Inc. All Rights Reserved.
+
+#pragma once
+
+#include "LiveLinkLensTypes.h"
+
+#include "LiveLinkOpenTrackIOLiveLinkTypes.generated.h"
+
+/**
+ * Struct for static OpenTrackIO data
+ */
+USTRUCT(BlueprintType)
+struct FLiveLinkOpenTrackIOStaticData : public FLiveLinkLensStaticData
+{
+	GENERATED_BODY()
+};
+
+/**
+ * Struct for dynamic (per-frame) OpenTrackIO data
+ */
+USTRUCT(BlueprintType)
+struct FLiveLinkOpenTrackIOFrameData : public FLiveLinkLensFrameData
+{
+	GENERATED_BODY()
+};
+
+/**
+ * Facility structure to handle OpenTrackIO data in blueprint
+ */
+USTRUCT(BlueprintType)
+struct FLiveLinkOpenTrackIOBlueprintData : public FLiveLinkBaseBlueprintData
+{
+	GENERATED_BODY()
+		
+	/** Static data that should not change every frame */
+	UPROPERTY(EditAnywhere, BlueprintReadWrite, Category = "LiveLink")
+	FLiveLinkOpenTrackIOStaticData StaticData;
+
+	/** Dynamic data that can change every frame  */
+	UPROPERTY(EditAnywhere, BlueprintReadWrite, Category = "LiveLink")
+	FLiveLinkOpenTrackIOFrameData FrameData;
+};
diff --git a/Engine/Plugins/VirtualProduction/LiveLinkOpenTrackIO/Source/LiveLinkOpenTrackIO/Private/LiveLinkOpenTrackIORole.cpp b/Engine/Plugins/VirtualProduction/LiveLinkOpenTrackIO/Source/LiveLinkOpenTrackIO/Private/LiveLinkOpenTrackIORole.cpp
new file mode 100644
--- /dev/null
+++ b/Engine/Plugins/VirtualProduction/LiveLinkOpenTrackIO/Source/LiveLinkOpenTrackIO/Private/LiveLinkOpenTrackIORole.cpp
@@ -0,0 +1,47 @@
+// Copyright Epic Games, Inc. All Rights Reserved.
+
+#include "LiveLinkOpenTrackIORole.h"
+
+#include "LiveLinkOpenTrackIORole.h"
+#include "LiveLinkOpenTrackIOLiveLinkTypes.h"
+
+#define LOCTEXT_NAMESPACE "LiveLinkOpenTrackIORole"
+
+UScriptStruct* ULiveLinkOpenTrackIORole::GetStaticDataStruct() const
+{
+	return FLiveLinkOpenTrackIOStaticData::StaticStruct();
+}
+
+UScriptStruct* ULiveLinkOpenTrackIORole::GetFrameDataStruct() const
+{
+	return FLiveLinkOpenTrackIOFrameData::StaticStruct();
+}
+
+UScriptStruct* ULiveLinkOpenTrackIORole::GetBlueprintDataStruct() const
+{
+	return FLiveLinkOpenTrackIOBlueprintData::StaticStruct();
+}
+
+bool ULiveLinkOpenTrackIORole::InitializeBlueprintData(const FLiveLinkSubjectFrameData& InSourceData, FLiveLinkBlueprintDataStruct& OutBlueprintData) const
+{
+	bool bSuccess = false;
+
+	FLiveLinkOpenTrackIOBlueprintData* BlueprintData = OutBlueprintData.Cast<FLiveLinkOpenTrackIOBlueprintData>();
+	const FLiveLinkOpenTrackIOStaticData* StaticData = InSourceData.StaticData.Cast<FLiveLinkOpenTrackIOStaticData>();
+	const FLiveLinkOpenTrackIOFrameData* FrameData = InSourceData.FrameData.Cast<FLiveLinkOpenTrackIOFrameData>();
+	if (BlueprintData && StaticData && FrameData)
+	{
+		GetStaticDataStruct()->CopyScriptStruct(&BlueprintData->StaticData, StaticData);
+		GetFrameDataStruct()->CopyScriptStruct(&BlueprintData->FrameData, FrameData);
+		bSuccess = true;
+	}
+
+	return bSuccess;
+}
+
+FText ULiveLinkOpenTrackIORole::GetDisplayName() const
+{
+	return LOCTEXT("OpenTrackIORole", "OpenTrackIO");
+}
+
+#undef LOCTEXT_NAMESPACE
diff --git a/Engine/Plugins/VirtualProduction/LiveLinkOpenTrackIO/Source/LiveLinkOpenTrackIO/Private/LiveLinkOpenTrackIORole.h b/Engine/Plugins/VirtualProduction/LiveLinkOpenTrackIO/Source/LiveLinkOpenTrackIO/Private/LiveLinkOpenTrackIORole.h
new file mode 100644
--- /dev/null
+++ b/Engine/Plugins/VirtualProduction/LiveLinkOpenTrackIO/Source/LiveLinkOpenTrackIO/Private/LiveLinkOpenTrackIORole.h
@@ -0,0 +1,28 @@
+// Copyright Epic Games, Inc. All Rights Reserved.
+
+#pragma once
+
+#include "LiveLinkLensRole.h"
+
+#include "LiveLinkOpenTrackIORole.generated.h"
+
+/**
+ * Role associated with OpenTrackIO data
+ */
+UCLASS(BlueprintType, meta = (DisplayName = "OpenTrackIO Role"))
+class ULiveLinkOpenTrackIORole : public ULiveLinkLensRole
+{
+	GENERATED_BODY()
+
+public:
+	//~ Begin ULiveLinkRole interface
+	virtual UScriptStruct* GetStaticDataStruct() const override;
+	virtual UScriptStruct* GetFrameDataStruct() const override;
+	virtual UScriptStruct* GetBlueprintDataStruct() const override;
+
+	virtual bool InitializeBlueprintData(const FLiveLinkSubjectFrameData& InSourceData, FLiveLinkBlueprintDataStruct& OutBlueprintData) const override;
+
+	virtual FText GetDisplayName() const override;
+	//~ End ULiveLinkRole interface
+};
+
diff --git a/Engine/Plugins/VirtualProduction/LiveLinkOpenTrackIO/Source/LiveLinkOpenTrackIO/Private/LiveLinkOpenTrackIOSource.cpp b/Engine/Plugins/VirtualProduction/LiveLinkOpenTrackIO/Source/LiveLinkOpenTrackIO/Private/LiveLinkOpenTrackIOSource.cpp
--- a/Engine/Plugins/VirtualProduction/LiveLinkOpenTrackIO/Source/LiveLinkOpenTrackIO/Private/LiveLinkOpenTrackIOSource.cpp
+++ b/Engine/Plugins/VirtualProduction/LiveLinkOpenTrackIO/Source/LiveLinkOpenTrackIO/Private/LiveLinkOpenTrackIOSource.cpp
@@ -14,7 +14,7 @@
 #include "LiveLinkOpenTrackIOUdpSocketReceiver.h"
 
 #include "Roles/LiveLinkCameraRole.h"
-#include "LiveLinkLensRole.h"
+#include "LiveLinkOpenTrackIORole.h"
 
 #include "HAL/Platform.h"
 #include "Misc/CoreDelegates.h"
@@ -417,7 +417,7 @@ void FLiveLinkOpenTrackIOSource::PushDataToLiveLink_AnyThread(const FLiveLinkOpe
 		if (bSubjectNameChanged || bStaticDataChanged)
 		{
 			Cache->StaticData.InitializeWith(StaticData);
-			Client->PushSubjectStaticData_AnyThread({ SourceGuid, Cache->SubjectName }, ULiveLinkLensRole::StaticClass(), MoveTemp(StaticData));
+			Client->PushSubjectStaticData_AnyThread({ SourceGuid, Cache->SubjectName }, ULiveLinkOpenTrackIORole::StaticClass(), MoveTemp(StaticData));
 		}
 	}
 
diff --git a/Engine/Plugins/VirtualProduction/LiveLinkOpenTrackIO/Source/LiveLinkOpenTrackIO/Private/LiveLinkOpenTrackIOTranscoder.cpp b/Engine/Plugins/VirtualProduction/LiveLinkOpenTrackIO/Source/LiveLinkOpenTrackIO/Private/LiveLinkOpenTrackIOTranscoder.cpp
--- a/Engine/Plugins/VirtualProduction/LiveLinkOpenTrackIO/Source/LiveLinkOpenTrackIO/Private/LiveLinkOpenTrackIOTranscoder.cpp
+++ b/Engine/Plugins/VirtualProduction/LiveLinkOpenTrackIO/Source/LiveLinkOpenTrackIO/Private/LiveLinkOpenTrackIOTranscoder.cpp
@@ -2,9 +2,10 @@
 // 
 #include "LiveLinkOpenTrackIOTranscoder.h"
 
-#include "LiveLinkLensTypes.h"
+#include "LiveLinkOpenTrackIOLiveLinkTypes.h"
 #include "LiveLinkOpenTrackIOConnectionSettings.h"
 #include "LiveLinkOpenTrackIOConversions.h"
+
 #include "Misc/App.h"
 #include "Misc/QualifiedFrameTime.h"
 #include "Roles/LiveLinkCameraTypes.h"
@@ -66,8 +67,8 @@ FLiveLinkStaticDataStruct FLiveLinkOpenTrackIOCache::MakeStaticData(const FLiveL
 {
 	/** Lens class inherits from the camera data so we use that as our base for the static struct. */
 
-	FLiveLinkStaticDataStruct StaticDataStruct(FLiveLinkLensStaticData::StaticStruct());
-	FLiveLinkLensStaticData& NewStaticData = *StaticDataStruct.Cast<FLiveLinkLensStaticData>();
+	FLiveLinkStaticDataStruct StaticDataStruct(FLiveLinkOpenTrackIOStaticData::StaticStruct());
+	FLiveLinkOpenTrackIOStaticData& NewStaticData = *StaticDataStruct.Cast<FLiveLinkOpenTrackIOStaticData>();
 
 	if (Data.Lens.PinholeFocalLength.IsSet())
 	{
@@ -176,8 +177,8 @@ FLiveLinkFrameDataStruct FLiveLinkOpenTrackIOCache::MakeFrameData(const FLiveLin
 	const FLiveLinkOpenTrackIOStaticCamera* Camera = StaticCamera.IsSet() ? &StaticCamera.GetValue() : nullptr;
 	const FLiveLinkOpenTrackIOStaticLens* Lens = StaticLens.IsSet() ? &StaticLens.GetValue() : nullptr;
 		
-	FLiveLinkFrameDataStruct FrameDataStruct = FLiveLinkFrameDataStruct(FLiveLinkLensFrameData::StaticStruct());
-	FLiveLinkLensFrameData& CameraAndLensData = *FrameDataStruct.Cast<FLiveLinkLensFrameData>();
+	FLiveLinkFrameDataStruct FrameDataStruct = FLiveLinkFrameDataStruct(FLiveLinkOpenTrackIOFrameData::StaticStruct());
+	FLiveLinkOpenTrackIOFrameData& CameraAndLensData = *FrameDataStruct.Cast<FLiveLinkOpenTrackIOFrameData>();
 
 	ToUnrealLens(CameraAndLensData, &Data.Lens, Camera);
 
diff --git a/Engine/Plugins/VirtualProduction/LiveLinkOpenTrackIO/Source/LiveLinkOpenTrackIO/Private/LiveLinkOpenTrackIOTypes.h b/Engine/Plugins/VirtualProduction/LiveLinkOpenTrackIO/Source/LiveLinkOpenTrackIO/Private/LiveLinkOpenTrackIOTypes.h
--- a/Engine/Plugins/VirtualProduction/LiveLinkOpenTrackIO/Source/LiveLinkOpenTrackIO/Private/LiveLinkOpenTrackIOTypes.h
+++ b/Engine/Plugins/VirtualProduction/LiveLinkOpenTrackIO/Source/LiveLinkOpenTrackIO/Private/LiveLinkOpenTrackIOTypes.h
@@ -12,8 +12,6 @@
 
 #include "LiveLinkOpenTrackIOTypes.generated.h"
 
-
-
 namespace UE::OpenTrackIO
 {
 	/** Brown-Conrady D-U, the default for nameless distortion model per opentrackio.org */
@@ -398,19 +396,19 @@ struct FLiveLinkOpenTrackIOStaticDuration
  * Static data from a parsed OpenTrackIO datagram. Per the spec this will be sent periodically from the source.
  */
 USTRUCT(BlueprintType)
-struct FLiveLinkOpenTrackIOStaticData
+struct FLiveLinkOpenTrackIOStatics
 {
 	GENERATED_BODY()
 
 	/** Returns true if the struct is equivalent to the default struct values. */
 	bool IsDefault() const
 	{
-		static FLiveLinkOpenTrackIOStaticData DefaultValue;
+		static FLiveLinkOpenTrackIOStatics DefaultValue;
 		return DefaultValue == *this;
 	}
 	
-	bool operator==(const FLiveLinkOpenTrackIOStaticData&) const = default;
-	bool operator!=(const FLiveLinkOpenTrackIOStaticData&) const = default;
+	bool operator==(const FLiveLinkOpenTrackIOStatics&) const = default;
+	bool operator!=(const FLiveLinkOpenTrackIOStatics&) const = default;
 	
 	UPROPERTY(EditAnywhere, BlueprintReadWrite, Category = "OpenTrackIO")
 	FLiveLinkOpenTrackIOStaticDuration Duration;
@@ -1285,7 +1283,7 @@ struct FLiveLinkOpenTrackIOData
 	bool operator!=(const FLiveLinkOpenTrackIOData&) const = default;
 	
 	UPROPERTY(EditAnywhere, BlueprintReadWrite, Category = "OpenTrackIO")
-	FLiveLinkOpenTrackIOStaticData Static;
+	FLiveLinkOpenTrackIOStatics Static;
 
 	UPROPERTY(EditAnywhere, BlueprintReadWrite, Category = "OpenTrackIO")
 	FLiveLinkOpenTrackIOTracker Tracker;
