- Update the LiveLinkOpenTrackIO source to support emitting the transform data as subjects by using the source settings object.
- Expose the OpenTrackIO data via the LiveLink FrameData so that it can be inspected in Blueprints.

#rnx

diff --git a/Engine/Plugins/VirtualProduction/LiveLinkOpenTrackIO/Source/LiveLinkOpenTrackIO/Private/LiveLinkOpenTrackIOConversions.h b/Engine/Plugins/VirtualProduction/LiveLinkOpenTrackIO/Source/LiveLinkOpenTrackIO/Private/LiveLinkOpenTrackIOConversions.h
--- a/Engine/Plugins/VirtualProduction/LiveLinkOpenTrackIO/Source/LiveLinkOpenTrackIO/Private/LiveLinkOpenTrackIOConversions.h
+++ b/Engine/Plugins/VirtualProduction/LiveLinkOpenTrackIO/Source/LiveLinkOpenTrackIO/Private/LiveLinkOpenTrackIOConversions.h
@@ -122,25 +122,25 @@ namespace LiveLinkOpenTrackIOConversions
 	}
 
 	/** From OpenTrackIO to Unreal Translation: Swap X <-> Y and convert meters to cm */
-	FVector ToUnrealTranslation(const FLiveLinkOpenTrackIO_XYZ& InXYZ)
+	inline FVector ToUnrealTranslation(const FLiveLinkOpenTrackIO_XYZ& InXYZ)
 	{
 		return FVector(InXYZ.Y, InXYZ.X, InXYZ.Z) * MetersToCentimeters;
 	}
 
 	/** From OpenTrackIO to Unreal Scale: Swap X <-> Y (unitless) */
-	FVector ToUnrealScale(const FLiveLinkOpenTrackIO_XYZ& InXYZ)
+	inline FVector ToUnrealScale(const FLiveLinkOpenTrackIO_XYZ& InXYZ)
 	{
 		return FVector(InXYZ.Y, InXYZ.X, InXYZ.Z);
 	}
 
 	/** From OpenTrackIO to Unreal Translation: Invert sign of Yaw */
-	FRotator ToUnrealRotation(const FLiveLinkOpenTrackIO_Rotator& InRotator)
+	inline FRotator ToUnrealRotation(const FLiveLinkOpenTrackIO_Rotator& InRotator)
 	{
 		return FRotator(InRotator.Tilt, -InRotator.Pan, InRotator.Roll);
 	}
 
 	/** From OpenTrackIO to Unreal FTransform */
-	FTransform ToUnrealTransform(const FLiveLinkOpenTrackIOTransform& InTransform)
+	inline FTransform ToUnrealTransform(const FLiveLinkOpenTrackIOTransform& InTransform)
 	{
 		return FTransform(
 			ToUnrealRotation(InTransform.Rotation),
diff --git a/Engine/Plugins/VirtualProduction/LiveLinkOpenTrackIO/Source/LiveLinkOpenTrackIO/Private/LiveLinkOpenTrackIOLiveLinkTypes.h b/Engine/Plugins/VirtualProduction/LiveLinkOpenTrackIO/Source/LiveLinkOpenTrackIO/Private/LiveLinkOpenTrackIOLiveLinkTypes.h
--- a/Engine/Plugins/VirtualProduction/LiveLinkOpenTrackIO/Source/LiveLinkOpenTrackIO/Private/LiveLinkOpenTrackIOLiveLinkTypes.h
+++ b/Engine/Plugins/VirtualProduction/LiveLinkOpenTrackIO/Source/LiveLinkOpenTrackIO/Private/LiveLinkOpenTrackIOLiveLinkTypes.h
@@ -4,6 +4,8 @@
 
 #include "LiveLinkLensTypes.h"
 
+#include "LiveLinkOpenTrackIOTypes.h"
+
 #include "LiveLinkOpenTrackIOLiveLinkTypes.generated.h"
 
 /**
@@ -22,6 +24,10 @@ USTRUCT(BlueprintType)
 struct FLiveLinkOpenTrackIOFrameData : public FLiveLinkLensFrameData
 {
 	GENERATED_BODY()
+
+	/* Access to the received OpenTrackIO data */
+	UPROPERTY(VisibleAnywhere, BlueprintReadWrite, Category = "OpenTrackIO")
+	FLiveLinkOpenTrackIOData OpenTrackData;
 };
 
 /**
diff --git a/Engine/Plugins/VirtualProduction/LiveLinkOpenTrackIO/Source/LiveLinkOpenTrackIO/Private/LiveLinkOpenTrackIOSource.cpp b/Engine/Plugins/VirtualProduction/LiveLinkOpenTrackIO/Source/LiveLinkOpenTrackIO/Private/LiveLinkOpenTrackIOSource.cpp
--- a/Engine/Plugins/VirtualProduction/LiveLinkOpenTrackIO/Source/LiveLinkOpenTrackIO/Private/LiveLinkOpenTrackIOSource.cpp
+++ b/Engine/Plugins/VirtualProduction/LiveLinkOpenTrackIO/Source/LiveLinkOpenTrackIO/Private/LiveLinkOpenTrackIOSource.cpp
@@ -4,9 +4,14 @@
 
 #include "Engine/Engine.h"
 
+#include "LiveLinkOpenTrackIOConversions.h"
+
 #include "LiveLinkTypes.h"
 #include "ILiveLinkClient.h"
 
+#include "LiveLinkOpenTrackIOTypes.h"
+#include "LiveLinkOpenTrackIOLiveLinkTypes.h"
+
 #include "LiveLinkOpenTrackIO.h"
 #include "LiveLinkOpenTrackIOSourceSettings.h"
 #include "LiveLinkOpenTrackIOParser.h"
@@ -21,6 +26,8 @@
 
 #include "Common/UdpSocketBuilder.h"
 
+#include "Roles/LiveLinkTransformRole.h"
+#include "Roles/LiveLinkTransformTypes.h"
 #include "Sockets.h"
 #include "SocketSubsystem.h"
 #include "Interfaces/IPv4/IPv4Address.h"
@@ -336,7 +343,8 @@ bool FLiveLinkOpenTrackIOSource::ConsumeInboundData()
 					TOptional<FLiveLinkOpenTrackIOData> Data = UE::OpenTrackIO::Private::ParsePayload(PayloadContainer);
 					if (Data)
 					{
-						PushDataToLiveLink_AnyThread(Header, *Data);
+						// This is a sink and we give up FLiveLinkOpenTrackData to optimize moving the OpenTrack data around.
+						PushDataToLiveLink_AnyThread(Header, MoveTemp(*Data));
 					}
 				}
 				// Remove the working payload from our table. This must happen 
@@ -351,7 +359,49 @@ bool FLiveLinkOpenTrackIOSource::ConsumeInboundData()
 	return bProcessedData;
 }
 
-void FLiveLinkOpenTrackIOSource::PushDataToLiveLink_AnyThread(const FLiveLinkOpenTrackIODatagramHeader& Header, const FLiveLinkOpenTrackIOData& InData)
+void FLiveLinkOpenTrackIOSource::RemoveAllTransformSubjects(FLiveLinkOpenTrackIOCache* Cache)
+{
+	// Remove any transform subjects based on this name change. 
+	for (FName TransformSubjectName : Cache->TransformSubjectNames)
+	{
+		Client->RemoveSubject_AnyThread({ SourceGuid, TransformSubjectName });
+	}
+	Cache->TransformSubjectNames.Reset();
+}
+
+void FLiveLinkOpenTrackIOSource::ConditionallyPushLiveLinkTransformData(FLiveLinkOpenTrackIOCache* Cache, const FLiveLinkOpenTrackIOData& InData)
+{
+	const bool bEnableTransformSubjects = IsValid(SavedSourceSettings) ? SavedSourceSettings->bEnableSubjectsPerTransform : false; 
+	if (!bEnableTransformSubjects)
+	{
+		RemoveAllTransformSubjects(Cache);
+		return;
+	}
+
+	for (const FLiveLinkOpenTrackIOTransform& Transform : InData.Transforms)
+	{
+		FName TransformName = Cache->GetTransformName(Transform);
+		if (!Cache->TransformSubjectNames.Contains(TransformName))
+		{
+			FLiveLinkStaticDataStruct StaticData(FLiveLinkTransformStaticData::StaticStruct());
+			FLiveLinkTransformStaticData& NewStaticData = *StaticData.Cast<FLiveLinkTransformStaticData>();
+			NewStaticData.bIsLocationSupported = true;
+			NewStaticData.bIsRotationSupported = true;
+			NewStaticData.bIsScaleSupported = true;
+			
+			Client->PushSubjectStaticData_AnyThread({ SourceGuid, TransformName }, ULiveLinkTransformRole::StaticClass(), MoveTemp(StaticData));
+		}
+
+		Cache->TransformSubjectNames.Add(TransformName);
+		FLiveLinkFrameDataStruct FrameData(FLiveLinkTransformFrameData::StaticStruct());
+		FLiveLinkTransformFrameData& NewFrameData = *FrameData.Cast<FLiveLinkTransformFrameData>();
+
+		NewFrameData.Transform = LiveLinkOpenTrackIOConversions::ToUnrealTransform(Transform);
+		Client->PushSubjectFrameData_AnyThread({SourceGuid,TransformName}, MoveTemp(FrameData));
+	}
+}
+
+void FLiveLinkOpenTrackIOSource::PushDataToLiveLink_AnyThread(const FLiveLinkOpenTrackIODatagramHeader& Header, FLiveLinkOpenTrackIOData InData)
 {
 	TRACE_CPUPROFILER_EVENT_SCOPE(OpenTrackIO::PushDataToLiveLink_AnyThread);
 
@@ -402,16 +452,19 @@ void FLiveLinkOpenTrackIOSource::PushDataToLiveLink_AnyThread(const FLiveLinkOpe
 			// Remove the previous Live Link subject, since the name is different.
 			if (!Cache->SubjectName.IsNone())
 			{
-				Client->RemoveSubject_AnyThread({ SourceGuid, Cache->SubjectName });				
+				Client->RemoveSubject_AnyThread({ SourceGuid, Cache->SubjectName });
 			}
+			RemoveAllTransformSubjects(Cache);
 
 			Cache->SubjectName = SubjectName;
 		}
 	}
 
+	const bool bShouldApplyTransform = IsValid(SavedSourceSettings) && !SavedSourceSettings->bEnableSubjectsPerTransform;
+	
 	// Update the static data if it has changed or if the subject name was updated.
 	{
-		FLiveLinkStaticDataStruct StaticData = Cache->MakeStaticData(InData);
+		FLiveLinkStaticDataStruct StaticData = Cache->MakeStaticData(InData, bShouldApplyTransform);
 		const bool bStaticDataChanged = Cache->StaticData != StaticData;
 
 		if (bSubjectNameChanged || bStaticDataChanged)
@@ -420,10 +473,18 @@ void FLiveLinkOpenTrackIOSource::PushDataToLiveLink_AnyThread(const FLiveLinkOpe
 			Client->PushSubjectStaticData_AnyThread({ SourceGuid, Cache->SubjectName }, ULiveLinkOpenTrackIORole::StaticClass(), MoveTemp(StaticData));
 		}
 	}
+	
+	ConditionallyPushLiveLinkTransformData(Cache, InData);
 
 	// Push the per-frame data
-	Client->PushSubjectFrameData_AnyThread({ SourceGuid, Cache->SubjectName }, Cache->MakeFrameData(InData));
+	FLiveLinkFrameDataStruct FrameData = Cache->MakeFrameData(InData, bShouldApplyTransform);
+	FLiveLinkOpenTrackIOFrameData& OpenTrackIOData = *FrameData.Cast<FLiveLinkOpenTrackIOFrameData>();
 
+	// Note we are moving the InData into the LL data.
+	OpenTrackIOData.OpenTrackData = MoveTemp(InData);
+	
+	Client->PushSubjectFrameData_AnyThread(	{ SourceGuid, Cache->SubjectName },	MoveTemp(FrameData) );
+	
 	Cache->UpdateLastKnownSequenceNumber(Header.SequenceNumber);
 }
 
diff --git a/Engine/Plugins/VirtualProduction/LiveLinkOpenTrackIO/Source/LiveLinkOpenTrackIO/Private/LiveLinkOpenTrackIOSource.h b/Engine/Plugins/VirtualProduction/LiveLinkOpenTrackIO/Source/LiveLinkOpenTrackIO/Private/LiveLinkOpenTrackIOSource.h
--- a/Engine/Plugins/VirtualProduction/LiveLinkOpenTrackIO/Source/LiveLinkOpenTrackIO/Private/LiveLinkOpenTrackIOSource.h
+++ b/Engine/Plugins/VirtualProduction/LiveLinkOpenTrackIO/Source/LiveLinkOpenTrackIO/Private/LiveLinkOpenTrackIOSource.h
@@ -111,9 +111,15 @@ private:
 	bool ConsumeInboundData();
 
 	/** Copy the open track data into Live Link equivalent data. This will establish any static data if the InData specifies it. */
-	void PushDataToLiveLink_AnyThread(const FLiveLinkOpenTrackIODatagramHeader& Header, const FLiveLinkOpenTrackIOData& InData);
-	
+	void PushDataToLiveLink_AnyThread(const FLiveLinkOpenTrackIODatagramHeader& Header, FLiveLinkOpenTrackIOData InData);
+
 private:
+	/** Remove all transform subjects from the LL source. */
+	void RemoveAllTransformSubjects(FLiveLinkOpenTrackIOCache* Cache);
+
+	/** Push transform data to Live Link if the user has asked for it in the SavedSourceSettings. New subjects will get created automatically. */
+	void ConditionallyPushLiveLinkTransformData(FLiveLinkOpenTrackIOCache* Cache, const FLiveLinkOpenTrackIOData& InData);
+	
 	struct FInboundMessage
 	{
 		/** Holds the segment data. */
diff --git a/Engine/Plugins/VirtualProduction/LiveLinkOpenTrackIO/Source/LiveLinkOpenTrackIO/Private/LiveLinkOpenTrackIOTranscoder.cpp b/Engine/Plugins/VirtualProduction/LiveLinkOpenTrackIO/Source/LiveLinkOpenTrackIO/Private/LiveLinkOpenTrackIOTranscoder.cpp
--- a/Engine/Plugins/VirtualProduction/LiveLinkOpenTrackIO/Source/LiveLinkOpenTrackIO/Private/LiveLinkOpenTrackIOTranscoder.cpp
+++ b/Engine/Plugins/VirtualProduction/LiveLinkOpenTrackIO/Source/LiveLinkOpenTrackIO/Private/LiveLinkOpenTrackIOTranscoder.cpp
@@ -2,6 +2,7 @@
 // 
 #include "LiveLinkOpenTrackIOTranscoder.h"
 
+#include "Containers/UnrealString.h"
 #include "LiveLinkOpenTrackIOLiveLinkTypes.h"
 #include "LiveLinkOpenTrackIOConnectionSettings.h"
 #include "LiveLinkOpenTrackIOConversions.h"
@@ -63,7 +64,18 @@ FName FLiveLinkOpenTrackIOCache::GetSubjectNameFromData(const FString& InSubject
 	return FName(FString::Printf(TEXT("OTrIO_%s_%03lld"), *LastOfSourceId, InData.SourceNumber));
 }
 
-FLiveLinkStaticDataStruct FLiveLinkOpenTrackIOCache::MakeStaticData(const FLiveLinkOpenTrackIOData& Data)
+FName FLiveLinkOpenTrackIOCache::GetTransformName(const FLiveLinkOpenTrackIOTransform& InTransform)
+{
+	if (!SubjectName.IsNone())
+	{
+		return FName(FString::Printf(TEXT("%s_%s"), *SubjectName.ToString(), *InTransform.Id));
+	}
+
+	return FName(InTransform.Id);
+}
+
+
+FLiveLinkStaticDataStruct FLiveLinkOpenTrackIOCache::MakeStaticData(const FLiveLinkOpenTrackIOData& Data, bool bApplyXform)
 {
 	/** Lens class inherits from the camera data so we use that as our base for the static struct. */
 
@@ -85,7 +97,7 @@ FLiveLinkStaticDataStruct FLiveLinkOpenTrackIOCache::MakeStaticData(const FLiveL
 		NewStaticData.bIsFocusDistanceSupported = true;
 	}
 
-	if (Data.Transforms.Num() > 0 || Data.Lens.EntrancePupilOffset.IsSet())
+	if ((bApplyXform && Data.Transforms.Num() > 0) || Data.Lens.EntrancePupilOffset.IsSet())
 	{
 		NewStaticData.bIsLocationSupported = true;
 		NewStaticData.bIsRotationSupported = true;
@@ -135,15 +147,18 @@ FLiveLinkStaticDataStruct FLiveLinkOpenTrackIOCache::MakeStaticData(const FLiveL
 }
 
 
-FLiveLinkFrameDataStruct FLiveLinkOpenTrackIOCache::MakeFrameData(const FLiveLinkOpenTrackIOData& Data)
+FLiveLinkFrameDataStruct FLiveLinkOpenTrackIOCache::MakeFrameData(const FLiveLinkOpenTrackIOData& Data, bool bApplyXform)
 {
 	using namespace LiveLinkOpenTrackIOConversions;
-	
-	/** Compute the final transform by concatentating the transforms in the frame data. */
+
 	FTransform FinalTransform; 
-	for (const FLiveLinkOpenTrackIOTransform& Transform : Data.Transforms)
+	if (bApplyXform)
 	{
-		FinalTransform *= ToUnrealTransform(Transform);
+		/** Compute the final transform by concatentating the transforms in the frame data. */
+		for (const FLiveLinkOpenTrackIOTransform& Transform : Data.Transforms)
+		{
+			FinalTransform *= ToUnrealTransform(Transform);
+		}
 	}
 
 	// Apply entrance pupil offset as a final transform in the forward direction.
@@ -191,6 +206,6 @@ FLiveLinkFrameDataStruct FLiveLinkOpenTrackIOCache::MakeFrameData(const FLiveLin
 	{
 		CameraAndLensData.MetaData.StringMetaData.Add(FName(Field.Key), Field.Value);
 	}
-	
+
 	return FrameDataStruct;
 }
diff --git a/Engine/Plugins/VirtualProduction/LiveLinkOpenTrackIO/Source/LiveLinkOpenTrackIO/Private/LiveLinkOpenTrackIOTranscoder.h b/Engine/Plugins/VirtualProduction/LiveLinkOpenTrackIO/Source/LiveLinkOpenTrackIO/Private/LiveLinkOpenTrackIOTranscoder.h
--- a/Engine/Plugins/VirtualProduction/LiveLinkOpenTrackIO/Source/LiveLinkOpenTrackIO/Private/LiveLinkOpenTrackIOTranscoder.h
+++ b/Engine/Plugins/VirtualProduction/LiveLinkOpenTrackIO/Source/LiveLinkOpenTrackIO/Private/LiveLinkOpenTrackIOTranscoder.h
@@ -6,8 +6,7 @@
 
 #include "LiveLinkTypes.h"
 
-#include "Containers/Map.h"
-#include "LiveLinkLensTypes.h"
+#include "Containers/Set.h"
 #include "Misc/Optional.h"
 
 struct FLiveLinkOpenTrackIOCache
@@ -15,6 +14,9 @@ struct FLiveLinkOpenTrackIOCache
 	/** Last known Subject Name */
 	FName SubjectName;
 
+	/** List of transform subject names known. */
+	TSet<FName> TransformSubjectNames;
+	
 	/** Live Link Static Data, used to detect changes */
 	FLiveLinkStaticDataStruct StaticData;
 
@@ -24,6 +26,9 @@ struct FLiveLinkOpenTrackIOCache
 	/** OpenTrackIO static lens data */
 	TOptional<FLiveLinkOpenTrackIOStaticLens> StaticLens;
 
+	/** Cached state of exposing the transforms as subjects. */
+	bool bSendTransformsAsSubjects = false;
+	
 	/**
 	 * Check the given sequence number and return true if it is the expected number. Sequence numbers should be incrementing
 	 * modulo UINT16_MAX.  This also accounts for delays in transmission based on sample rate provided by the first inbound packets.
@@ -45,12 +50,15 @@ struct FLiveLinkOpenTrackIOCache
 	
 	/** Constructs a subject name from the given data. */
 	FName GetSubjectNameFromData(const FString& Prefix, const FLiveLinkOpenTrackIOData& Data);
-	
+
+	/** Constructs a subject name from the given data. */
+	FName GetTransformName(const FLiveLinkOpenTrackIOTransform& InTransform);
+		
 	/** Create Live Link static data from the OpenTrackIO data */
-	FLiveLinkStaticDataStruct MakeStaticData(const FLiveLinkOpenTrackIOData& StaticData);
+	FLiveLinkStaticDataStruct MakeStaticData(const FLiveLinkOpenTrackIOData& StaticData, const bool bApplyXform);
 
 	/** Create Live Link the per-frame data from the OpenTrackIO data */
-	FLiveLinkFrameDataStruct MakeFrameData(const FLiveLinkOpenTrackIOData& Data);
+	FLiveLinkFrameDataStruct MakeFrameData(const FLiveLinkOpenTrackIOData& Data, const bool bApplyXform);
 
 private:
 	/** Tracks the current sequence number for the source. */
diff --git a/Engine/Plugins/VirtualProduction/LiveLinkOpenTrackIO/Source/LiveLinkOpenTrackIO/Public/LiveLinkOpenTrackIOSourceSettings.h b/Engine/Plugins/VirtualProduction/LiveLinkOpenTrackIO/Source/LiveLinkOpenTrackIO/Public/LiveLinkOpenTrackIOSourceSettings.h
--- a/Engine/Plugins/VirtualProduction/LiveLinkOpenTrackIO/Source/LiveLinkOpenTrackIO/Public/LiveLinkOpenTrackIOSourceSettings.h
+++ b/Engine/Plugins/VirtualProduction/LiveLinkOpenTrackIO/Source/LiveLinkOpenTrackIO/Public/LiveLinkOpenTrackIOSourceSettings.h
@@ -56,4 +56,11 @@ public:
 		EditConditionHides
 		))
 	FString AssignedUnicastEndpoint = TEXT("");
+
+
+	/**
+	 * Emit subjects for each transform included in the OpenTrackIO channel.
+	 */
+	UPROPERTY(EditAnywhere, Category=Subjects)
+	bool bEnableSubjectsPerTransform = false;
 };
