Change the transforms as subjects setting to an enum to give the user more control. Change the entrance pupil to non-optional value.

#rnx

diff --git a/Engine/Plugins/VirtualProduction/LiveLinkOpenTrackIO/Source/LiveLinkOpenTrackIO/Private/LiveLinkOpenTrackIOSource.cpp b/Engine/Plugins/VirtualProduction/LiveLinkOpenTrackIO/Source/LiveLinkOpenTrackIO/Private/LiveLinkOpenTrackIOSource.cpp
--- a/Engine/Plugins/VirtualProduction/LiveLinkOpenTrackIO/Source/LiveLinkOpenTrackIO/Private/LiveLinkOpenTrackIOSource.cpp
+++ b/Engine/Plugins/VirtualProduction/LiveLinkOpenTrackIO/Source/LiveLinkOpenTrackIO/Private/LiveLinkOpenTrackIOSource.cpp
@@ -371,7 +371,7 @@ void FLiveLinkOpenTrackIOSource::RemoveAllTransformSubjects(FLiveLinkOpenTrackIO
 
 void FLiveLinkOpenTrackIOSource::ConditionallyPushLiveLinkTransformData(FLiveLinkOpenTrackIOCache* Cache, const FLiveLinkOpenTrackIOData& InData)
 {
-	const bool bEnableTransformSubjects = IsValid(SavedSourceSettings) ? SavedSourceSettings->bEnableSubjectsPerTransform : false; 
+	const bool bEnableTransformSubjects = IsValid(SavedSourceSettings) ? SavedSourceSettings->ShouldExtractTransformSubjects() : false; 
 	if (!bEnableTransformSubjects)
 	{
 		RemoveAllTransformSubjects(Cache);
@@ -460,7 +460,7 @@ void FLiveLinkOpenTrackIOSource::PushDataToLiveLink_AnyThread(const FLiveLinkOpe
 		}
 	}
 
-	const bool bShouldApplyTransform = IsValid(SavedSourceSettings) && !SavedSourceSettings->bEnableSubjectsPerTransform;
+	const bool bShouldApplyTransform = IsValid(SavedSourceSettings) && SavedSourceSettings->ShouldApplyXformToCamera();
 	
 	// Update the static data if it has changed or if the subject name was updated.
 	{
diff --git a/Engine/Plugins/VirtualProduction/LiveLinkOpenTrackIO/Source/LiveLinkOpenTrackIO/Private/LiveLinkOpenTrackIOTranscoder.cpp b/Engine/Plugins/VirtualProduction/LiveLinkOpenTrackIO/Source/LiveLinkOpenTrackIO/Private/LiveLinkOpenTrackIOTranscoder.cpp
--- a/Engine/Plugins/VirtualProduction/LiveLinkOpenTrackIO/Source/LiveLinkOpenTrackIO/Private/LiveLinkOpenTrackIOTranscoder.cpp
+++ b/Engine/Plugins/VirtualProduction/LiveLinkOpenTrackIO/Source/LiveLinkOpenTrackIO/Private/LiveLinkOpenTrackIOTranscoder.cpp
@@ -97,7 +97,7 @@ FLiveLinkStaticDataStruct FLiveLinkOpenTrackIOCache::MakeStaticData(const FLiveL
 		NewStaticData.bIsFocusDistanceSupported = true;
 	}
 
-	if ((bApplyXform && Data.Transforms.Num() > 0) || Data.Lens.EntrancePupilOffset.IsSet())
+	if ((bApplyXform && Data.Transforms.Num() > 0))
 	{
 		NewStaticData.bIsLocationSupported = true;
 		NewStaticData.bIsRotationSupported = true;
@@ -162,18 +162,15 @@ FLiveLinkFrameDataStruct FLiveLinkOpenTrackIOCache::MakeFrameData(const FLiveLin
 	}
 
 	// Apply entrance pupil offset as a final transform in the forward direction.
-	if (Data.Lens.EntrancePupilOffset.IsSet())
-	{
-		const FLiveLinkOpenTrackIO_XYZ EntrancePupilAsXYZ(0, Data.Lens.EntrancePupilOffset.GetValue(), 0);
+	const FLiveLinkOpenTrackIO_XYZ EntrancePupilAsXYZ(0, Data.Lens.EntrancePupilOffset, 0);
 
-		const FTransform PupilOffset(
-			FQuat::Identity,
-			ToUnrealTranslation(EntrancePupilAsXYZ),
-			FVector(1, 1, 1)
-		);
+	const FTransform PupilOffset(
+		FQuat::Identity,
+		ToUnrealTranslation(EntrancePupilAsXYZ),
+		FVector(1, 1, 1)
+	);
 
-		FinalTransform *= PupilOffset;
-	}
+	FinalTransform *= PupilOffset;
 
 	FQualifiedFrameTime SampleTC;
 	if (Data.Timing.IsDefault())
diff --git a/Engine/Plugins/VirtualProduction/LiveLinkOpenTrackIO/Source/LiveLinkOpenTrackIO/Private/LiveLinkOpenTrackIOTypes.h b/Engine/Plugins/VirtualProduction/LiveLinkOpenTrackIO/Source/LiveLinkOpenTrackIO/Private/LiveLinkOpenTrackIOTypes.h
--- a/Engine/Plugins/VirtualProduction/LiveLinkOpenTrackIO/Source/LiveLinkOpenTrackIO/Private/LiveLinkOpenTrackIOTypes.h
+++ b/Engine/Plugins/VirtualProduction/LiveLinkOpenTrackIO/Source/LiveLinkOpenTrackIO/Private/LiveLinkOpenTrackIOTypes.h
@@ -996,7 +996,7 @@ struct FLiveLinkOpenTrackIOLens
 	 * virtual camera’s transform chain.
 	 */
 	UPROPERTY(EditAnywhere, BlueprintReadWrite, Category = "OpenTrackIO")
-	TOptional<float> EntrancePupilOffset;
+	float EntrancePupilOffset = 0.0f;
 
 	/**
 	 * Coefficients for calculating the exposure fall-off (vignetting) of a lens
diff --git a/Engine/Plugins/VirtualProduction/LiveLinkOpenTrackIO/Source/LiveLinkOpenTrackIO/Public/LiveLinkOpenTrackIOSourceSettings.h b/Engine/Plugins/VirtualProduction/LiveLinkOpenTrackIO/Source/LiveLinkOpenTrackIO/Public/LiveLinkOpenTrackIOSourceSettings.h
--- a/Engine/Plugins/VirtualProduction/LiveLinkOpenTrackIO/Source/LiveLinkOpenTrackIO/Public/LiveLinkOpenTrackIOSourceSettings.h
+++ b/Engine/Plugins/VirtualProduction/LiveLinkOpenTrackIO/Source/LiveLinkOpenTrackIO/Public/LiveLinkOpenTrackIOSourceSettings.h
@@ -8,6 +8,14 @@
 
 #include "LiveLinkOpenTrackIOSourceSettings.generated.h"
 
+UENUM()
+enum class ELiveLinkOpenTrackIOTransformSubjects
+{
+	NoTransformSubjects					 UMETA(ToolTip = "Do not generate transform subjects from OpenTrackIO data."),
+	EnableTransformSubjects				 UMETA(ToolTip = "Generate transform subjects and calculate final transform on the camera frame data."),
+	EnableTransformSubjectsNoCameraXForm UMETA(ToolTip = "Generate transform subjects and do not calculate final transform on the camera frame data."),
+};
+
 UCLASS()
 class LIVELINKOPENTRACKIO_API ULiveLinkOpenTrackIOSourceSettings : public ULiveLinkSourceSettings
 {
@@ -57,10 +65,19 @@ public:
 		))
 	FString AssignedUnicastEndpoint = TEXT("");
 
-
 	/**
 	 * Emit subjects for each transform included in the OpenTrackIO channel.
 	 */
 	UPROPERTY(EditAnywhere, Category=Subjects)
-	bool bEnableSubjectsPerTransform = false;
+	ELiveLinkOpenTrackIOTransformSubjects SubjectsPerTransform = ELiveLinkOpenTrackIOTransformSubjects::NoTransformSubjects;
+
+	bool ShouldExtractTransformSubjects() const
+	{
+		return SubjectsPerTransform != ELiveLinkOpenTrackIOTransformSubjects::NoTransformSubjects;
+	}
+
+	bool ShouldApplyXformToCamera() const
+	{
+		return SubjectsPerTransform != ELiveLinkOpenTrackIOTransformSubjects::EnableTransformSubjectsNoCameraXForm;
+	}
 };
