OpenTrackIO: Fix crash when closing sources.

The issue was a race condition due to the udp receivers being created and closed in a worker thread but also closed and reset in the game thread. Opted for simplifying the threading and only managing the receivers in the game thread.

* Consume inbound packets now runs on the Udp Receiver  thread to favor lower latency. Processing incoming packets (OTrIO segments) is relatively light weight (parse cbor or json, push to live link) so this should be ok.
* Removed display of local endpoint since the port should be well defined when creating the subject.
* Connection state fully managed in the game thread.
* Socket and udp receiver lifetime fully managed in the game thread.
* Removed the custom Udp Receiver as the standard one suffices. It had only been added to be able to explicitly wait for stop to end before destruction, but this is no longer necessary as the root issue for prior crashes was this race condition.

#rnx

diff --git a/Engine/Plugins/VirtualProduction/LiveLinkOpenTrackIO/Source/LiveLinkOpenTrackIO/Private/LiveLinkOpenTrackIOSource.cpp b/Engine/Plugins/VirtualProduction/LiveLinkOpenTrackIO/Source/LiveLinkOpenTrackIO/Private/LiveLinkOpenTrackIOSource.cpp
--- a/Engine/Plugins/VirtualProduction/LiveLinkOpenTrackIO/Source/LiveLinkOpenTrackIO/Private/LiveLinkOpenTrackIOSource.cpp
+++ b/Engine/Plugins/VirtualProduction/LiveLinkOpenTrackIO/Source/LiveLinkOpenTrackIO/Private/LiveLinkOpenTrackIOSource.cpp
@@ -16,7 +16,6 @@
 #include "LiveLinkOpenTrackIOSourceSettings.h"
 #include "LiveLinkOpenTrackIOParser.h"
 #include "LiveLinkOpenTrackIOTranscoder.h"
-#include "LiveLinkOpenTrackIOUdpSocketReceiver.h"
 
 #include "Roles/LiveLinkCameraRole.h"
 #include "LiveLinkOpenTrackIORole.h"
@@ -25,6 +24,7 @@
 #include "Misc/CoreDelegates.h"
 
 #include "Common/UdpSocketBuilder.h"
+#include "Common/UdpSocketReceiver.h"
 
 #include "Roles/LiveLinkTransformRole.h"
 #include "Roles/LiveLinkTransformTypes.h"
@@ -37,7 +37,6 @@
 
 namespace UE::LiveLinkOpenTrackIO::Private
 {
-
 TOptional<FIPv4Endpoint> ParseEndpoint(const FString& InEndpointString)
 {
 	FIPv4Endpoint Endpoint;
@@ -70,12 +69,12 @@ void DoJoinMulticastGroup(const TSharedRef<FInternetAddr>& MulticastAddr, const
 	if (bJoinedGroup)
 	{
 		UE_LOG(LogLiveLinkOpenTrackIO, Display, TEXT("Added local interface '%s' to multicast group '%s'"),
-			   *IpAddr->ToString(false), *MulticastAddr->ToString(true));
+				*IpAddr->ToString(false), *MulticastAddr->ToString(true));
 	}
 	else
 	{
 		UE_LOG(LogLiveLinkOpenTrackIO, Warning, TEXT("Failed to join multicast group '%s' on detected local interface '%s'"),
-			   *MulticastAddr->ToString(true), *IpAddr->ToString(false));
+				*MulticastAddr->ToString(true), *IpAddr->ToString(false));
 	}
 }
 
@@ -121,26 +120,37 @@ FIPv4Endpoint MakeMulticastAddressFromSourceNumber(uint8 InSourceNumber, uint16
 	const uint8 SourceNumber = FMath::Clamp(InSourceNumber, 1, 200);
 	return FIPv4Endpoint(FIPv4Address(235, 135, 1, SourceNumber), InPortNumber);
 }
-
 }
 
 FLiveLinkOpenTrackIOSource::FLiveLinkOpenTrackIOSource(FLiveLinkOpenTrackIOConnectionSettings InConnectionSettings)
 	: ConnectionSettings(MoveTemp(InConnectionSettings))
 {
 	SourceType = LOCTEXT("SourceType_OpenTrackIO", "OpenTrack I/O");
-	SourceStatus = LOCTEXT("Startup", "Waiting for startup...");
+	SourceStatus = LOCTEXT("Initialization", "Initializing receivers...");
 
-	DeferredStartDelegateHandle = FCoreDelegates::OnEndFrame.AddRaw(this, &FLiveLinkOpenTrackIOSource::Start);
+	// We'll use this callback for cleanup fallback. See FLiveLinkOpenTrackIOSource::OnEnginePreExit for more details.
+	FCoreDelegates::OnEnginePreExit.AddRaw(this, &FLiveLinkOpenTrackIOSource::OnEnginePreExit);
+}
+
+void FLiveLinkOpenTrackIOSource::OnEnginePreExit()
+{
+	// This source uses latent shutdown (see RequestSourceShutdown()) but because Update() may not be called again to wait for
+	// the receiver threads to stop, it might be possible to access UObject subsystems from those threads that are being torn down.
+	// e.g. TBaseStructure<FLiveLinkOpenTrackIOData>::Get() will return invalid data and crash the Cbor parser.
+	//
+	// To avoid this deterministically, on EnginePreExit we request shutdown and make one last call to Update() which should
+	// wait for the udp receiver threads to end.
+
+	RequestSourceShutdown();
+	Update();
 }
 
 FLiveLinkOpenTrackIOSource::~FLiveLinkOpenTrackIOSource()
 {
 	CloseSockets();
 
-	if (DeferredStartDelegateHandle.IsValid())
-	{
-		FCoreDelegates::OnEndFrame.Remove(DeferredStartDelegateHandle);
-	}
+	// Remove our backup cleanup call.
+	FCoreDelegates::OnEnginePreExit.RemoveAll(this);
 }
 
 void FLiveLinkOpenTrackIOSource::ReceiveClient(ILiveLinkClient* InClient, FGuid InSourceGuid)
@@ -151,85 +161,135 @@ void FLiveLinkOpenTrackIOSource::ReceiveClient(ILiveLinkClient* InClient, FGuid
 
 bool FLiveLinkOpenTrackIOSource::IsSourceStillValid() const
 {
-	const bool bIsSourceValid = !Stopping && GameThreadConnectionState == ELiveLinkOpenTrackIOState::Receiving; 
-	return true;
+	return ConnectionState == ELiveLinkOpenTrackIOState::Receiving;
 }
 
 bool FLiveLinkOpenTrackIOSource::RequestSourceShutdown()
 {
-	Stop();
+	if (ConnectionState == ELiveLinkOpenTrackIOState::ShutDown)
+	{
+		return true;
+	}
 
-	return true;
+	// We do a latent shutdown because of a possible deadlock due to a mutex shared between:
+	// 
+	// * PushSubjectFrameData_AnyThread and
+	// * RequestSourceShutdown
+	//
+	// Since PushSubjectFrameData_AnyThread is called in the Udp receiver thread, it means we can't wait for that thread
+	// to stop here, because it may be waiting for the lock we're currently in.
+
+	// Stop Udp Receivers so that they stop pushing more packets unnecessarily.
+	StopUdpReceivers();
+
+	// This flag will be checked by the state machine, which will then clean up and enter the ShutDown state.
+	bShutdownRequested = true;
+
+	return false;
 }
 
 void FLiveLinkOpenTrackIOSource::Update()
 {
-	if (GameThreadConnectionState != ConnectionState)
+	// Any state can lead to ShutDown directly.
+	if (bShutdownRequested && (ConnectionState != ELiveLinkOpenTrackIOState::ShutDown))
 	{
-		GameThreadConnectionState = ConnectionState;
+		CloseSockets();
+		SetConnectionState(ELiveLinkOpenTrackIOState::ShutDown);
 
-		switch (GameThreadConnectionState)
-		{
-			case ELiveLinkOpenTrackIOState::ResetRequested:
-				SourceStatus = LOCTEXT("Resetting", "Resetting source.");
-				break;
-			case ELiveLinkOpenTrackIOState::EndpointsReady:
-				SourceStatus = LOCTEXT("EndpointsReady", "Starting socket setup.");
-				break;
-			case ELiveLinkOpenTrackIOState::Receiving:
-				break;
-			default:
-				SourceStatus = LOCTEXT("NotStarted", "Not started");
-		};
-	}
-
-	if (GameThreadConnectionState == ELiveLinkOpenTrackIOState::Receiving)
-	{
-		if (OpenTrackIOPacketCount > 0)
-		{
-			SourceStatus = LOCTEXT("Receiving", "Receiving.");
-		}
-		else
-		{
-			SourceStatus = LOCTEXT("WaitingForData", "Waiting for data.");
-		}
+		// Note: We purposedly do not clear the bShutdownRequested flag. It is one-time use only,
+		// which means that this source will never exit the ShutDown state.
 	}
-}
+	else if (bResetRequested)
+	{
+		// Any state, unless shutting down, can lead to ResetRequested.
+		SetConnectionState(ELiveLinkOpenTrackIOState::ResetRequested);
 
-void FLiveLinkOpenTrackIOSource::Start()
-{
-	check(DeferredStartDelegateHandle.IsValid());
+		bResetRequested = false;
+	}
 
-	FCoreDelegates::OnEndFrame.Remove(DeferredStartDelegateHandle);
-	DeferredStartDelegateHandle.Reset();
+	switch (ConnectionState)
+	{
+		case ELiveLinkOpenTrackIOState::NotStarted:
+			SourceStatus = LOCTEXT("NotStarted", "Not started");
 
-	SourceStatus = LOCTEXT("Initialization", "Initializing receivers...");
-	ThreadName = "LiveLinkOpenTrackIO Receiver";
-	ThreadName.AppendInt(FAsyncThreadIndex::GetNext());
+			// Note: InitializeSettings should get us out of this state directly to ResetRequested
+		break;
 
-	Thread.Reset(FRunnableThread::Create(this, *ThreadName, 128 * 1024, TPri_Normal, FPlatformAffinity::GetPoolThreadMask()));
-}
+		case ELiveLinkOpenTrackIOState::ResetRequested:
 
-void FLiveLinkOpenTrackIOSource::Stop()
-{
-	Stopping = true;
+			SourceStatus = LOCTEXT("Resetting", "Resetting source.");
+
+			// If the user has requested a connection reset. Close the socket and attempt to re-open on the next loop iteration.
+			CloseSockets();
+
+			LastDataReadTime = 0;
+
+			// Re-parse the endpoints to reset the connection.
+			if (ParseEndpointFromSourceSettings())
+			{
+				SetConnectionState(ELiveLinkOpenTrackIOState::EndpointsReady);
+			}
+
+			break;
+
+		case ELiveLinkOpenTrackIOState::EndpointsReady:
+			SourceStatus = LOCTEXT("EndpointsReady", "Starting socket setup.");
+
+			if (OpenSockets())
+			{
+				SetConnectionState(ELiveLinkOpenTrackIOState::Receiving);
+			}
+			else
+			{
+				SetConnectionState(ELiveLinkOpenTrackIOState::ResetRequested);
+			}
+
+			break;
+
+		case ELiveLinkOpenTrackIOState::Receiving:
+			if (FPlatformTime::Seconds() - LastDataReadTime < 1.0)
+			{
+				SourceStatus = LOCTEXT("Receiving", "Receiving.");
+			}
+			else
+			{
+				SourceStatus = LOCTEXT("WaitingForData", "Waiting for data.");
+			}
+			break;
+
+		case ELiveLinkOpenTrackIOState::ShutDown:
+			SourceStatus = LOCTEXT("ShutDown", "Shut Down");
+
+			// If we're here, then the udp receivers and sockets must have already been closed.
+			break;
+
+		default:
+			checkNoEntry();
+			break;
+	};
 }
 
-void FLiveLinkOpenTrackIOSource::ParseEndpointFromSourceSettings()
+bool FLiveLinkOpenTrackIOSource::ParseEndpointFromSourceSettings()
 {
 	check(IsValid(SavedSourceSettings));
 	
-	using namespace UE::LiveLinkOpenTrackIO::Private; 
-	if (TOptional<FIPv4Endpoint> InUniEndpoint = ParseEndpoint(SavedSourceSettings->UnicastEndpoint))
-	{
-		UnicastEndpoint = *InUniEndpoint;
+	using namespace UE::LiveLinkOpenTrackIO::Private;
 
-		// Note: The Unicast port number is not used when using multicast, only the address (which is used as an interface address to bind to).
-		UnicastEndpoint.Port = ConnectionSettings.UnicastPort;
+	TOptional<FIPv4Endpoint> InUniEndpoint = ParseEndpoint(SavedSourceSettings->UnicastEndpoint);
 
-		MulticastEndpoint = MakeMulticastAddressFromSourceNumber(ConnectionSettings.SourceNumber, SavedSourceSettings->MulticastPort);
-		ConnectionState = ELiveLinkOpenTrackIOState::EndpointsReady;
+	if (!InUniEndpoint.IsSet())
+	{
+		return false;
 	}
+
+	UnicastEndpoint = *InUniEndpoint;
+
+	// Note: The Unicast port number is not used when using multicast, only the address (which is used as an interface address to bind to).
+	UnicastEndpoint.Port = ConnectionSettings.UnicastPort;
+
+	MulticastEndpoint = MakeMulticastAddressFromSourceNumber(ConnectionSettings.SourceNumber, SavedSourceSettings->MulticastPort);
+
+	return true;
 }
 
 void FLiveLinkOpenTrackIOSource::InitializeSettings(ULiveLinkSourceSettings* Settings)
@@ -243,7 +303,7 @@ void FLiveLinkOpenTrackIOSource::InitializeSettings(ULiveLinkSourceSettings* Set
 
 		SavedSourceSettings = OpenTrackIOSettings;
 
-		ParseEndpointFromSourceSettings();
+		bResetRequested = true;
 	}
 	else
 	{
@@ -251,18 +311,6 @@ void FLiveLinkOpenTrackIOSource::InitializeSettings(ULiveLinkSourceSettings* Set
 	}
 }
 
-void FLiveLinkOpenTrackIOSource::ResetSockets()
-{
-	SourceStatus = LOCTEXT("Resetting", "Resetting source.");
-	ConnectionState = ELiveLinkOpenTrackIOState::ResetRequested;
-	if (IsValid(SavedSourceSettings))
-	{
-		SavedSourceSettings->AssignedUnicastEndpoint = TEXT("");
-	}
-
-	OpenTrackIOPacketCount = 0;
-}
-
 void FLiveLinkOpenTrackIOSource::OnSettingsChanged(ULiveLinkSourceSettings* Settings, const FPropertyChangedEvent& PropertyChangedEvent)
 {
 	ILiveLinkSource::OnSettingsChanged(Settings, PropertyChangedEvent);
@@ -280,35 +328,22 @@ void FLiveLinkOpenTrackIOSource::OnSettingsChanged(ULiveLinkSourceSettings* Sett
 				|| PropertyName == GET_MEMBER_NAME_CHECKED(ULiveLinkOpenTrackIOSourceSettings, MulticastPort))
 			{
 				// Start the reset process.
-				ResetSockets();
+				bResetRequested = true;
 			}
 		}
 	}
 }
 	
-double FLiveLinkOpenTrackIOSource::DevicePollWaitTime() const
-{
-	const double UpdateRateInHz = LocalUpdateRateInHz > 0 ? (double) LocalUpdateRateInHz : 60.0;
-	return 1.0 / UpdateRateInHz;
-}
-
 void FLiveLinkOpenTrackIOSource::HandleInboundData(const TSharedPtr<FArrayReader, ESPMode::ThreadSafe>& Data, const FIPv4Endpoint& InSender)
 {
-	if (Stopping)
-	{
-		return;
-	}
-
 	Inbound.Enqueue(FInboundMessage{ Data, InSender });
-
-	/** TODO Rather than a busy loop we should use FEvent to trigger thread wake-up. */
+	ConsumeInboundData();
 }
 
-bool FLiveLinkOpenTrackIOSource::ConsumeInboundData()
+void FLiveLinkOpenTrackIOSource::ConsumeInboundData()
 {
 	TRACE_CPUPROFILER_EVENT_SCOPE(OpenTrackIO::ConsumeInboundData);
 
-	bool bProcessedData = false;
 	FInboundMessage Message;
 
 	while (Inbound.Dequeue(Message))
@@ -332,8 +367,7 @@ bool FLiveLinkOpenTrackIOSource::ConsumeInboundData()
 				});
 			}
 
-			bProcessedData = true;
-			OpenTrackIOPacketCount++;
+			LastDataReadTime = FPlatformTime::Seconds();
 
 			if (PayloadContainer.GetHeader().IsLastSegment())
 			{
@@ -356,7 +390,6 @@ bool FLiveLinkOpenTrackIOSource::ConsumeInboundData()
 			UE_LOG(LogLiveLinkOpenTrackIO, Display, TEXT("Failed to handle inbound data."));
 		}
 	}
-	return bProcessedData;
 }
 
 void FLiveLinkOpenTrackIOSource::RemoveAllTransformSubjects(FLiveLinkOpenTrackIOCache* Cache)
@@ -488,31 +521,22 @@ void FLiveLinkOpenTrackIOSource::PushDataToLiveLink_AnyThread(const FLiveLinkOpe
 	Cache->UpdateLastKnownSequenceNumber(Header.SequenceNumber);
 }
 
-void FLiveLinkOpenTrackIOSource::CloseSockets()
+void FLiveLinkOpenTrackIOSource::StopUdpReceivers() 
 {
-	// Stop the receivers
-
 	if (MulticastReceiver)
 	{
 		MulticastReceiver->Stop();
 	}
-	
+
 	if (UnicastReceiver)
 	{
 		UnicastReceiver->Stop();
 	}
+}
 
-	// Wait until they are confirmed to not be using the sockets.
-
-	if (MulticastReceiver)
-	{
-		MulticastReceiver->WaitForStop();
-	}
-
-	if (UnicastReceiver)
-	{
-		UnicastReceiver->WaitForStop();
-	}
+void FLiveLinkOpenTrackIOSource::CloseSockets()
+{
+	StopUdpReceivers();
 
 	MulticastReceiver.Reset();
 	UnicastReceiver.Reset();
@@ -533,6 +557,8 @@ void FLiveLinkOpenTrackIOSource::CloseSockets()
 		UnicastSocket = nullptr;
 	}
 
+	// It should be safe to clear this because the Udp threads are guaranteed to be shut down at this point,
+	// since the call to .Reset() earlier in this fuction will block until they are destroyed.
 	OpenTrackIOCacheMap.Reset();
 }
 
@@ -576,7 +602,7 @@ bool FLiveLinkOpenTrackIOSource::OpenMulticastSocket()
 
 	const FTimespan ThreadWaitTime = FTimespan::FromMilliseconds(100);
 
-	MulticastReceiver = MakeUnique<FLiveLinkOpenTrackIOUdpSocketReceiver>(MulticastSocket, ThreadWaitTime, TEXT("LiveLinkOpenTransportIO_MulticastReceiver"));
+	MulticastReceiver = MakeUnique<FUdpSocketReceiver>(MulticastSocket, ThreadWaitTime, TEXT("LiveLinkOpenTransportIO_MulticastReceiver"));
 	MulticastReceiver->OnDataReceived().BindRaw(this, &FLiveLinkOpenTrackIOSource::HandleInboundData);
 	MulticastReceiver->SetMaxReadBufferSize(2048);
 	MulticastReceiver->Start();
@@ -605,19 +631,11 @@ bool FLiveLinkOpenTrackIOSource::OpenUnicastSocket()
 	FString IpPortStr = FString::Format(TEXT("{0}:{1}"), { UnicastEndpoint.Address.ToString(), PortNo });
 	UE_LOG(LogLiveLinkOpenTrackIO, Display, TEXT("OpenTrackIO unicast socket socket bound to '%s'."), *IpPortStr);
 
-	AsyncTask(ENamedThreads::GameThread, [IpPortStr, this]() {
-
-		if (IsValid(SavedSourceSettings))
-		{
-			SavedSourceSettings->AssignedUnicastEndpoint = IpPortStr;
-		}
-		});
-
 	// Create receiver for socket above
 
 	const FTimespan ThreadWaitTime = FTimespan::FromMilliseconds(100);
 
-	UnicastReceiver = MakeUnique<FLiveLinkOpenTrackIOUdpSocketReceiver>(UnicastSocket, ThreadWaitTime, TEXT("LiveLinkOpenTransportIO_UnicastReceiver"));
+	UnicastReceiver = MakeUnique<FUdpSocketReceiver>(UnicastSocket, ThreadWaitTime, TEXT("LiveLinkOpenTransportIO_UnicastReceiver"));
 	UnicastReceiver->OnDataReceived().BindRaw(this, &FLiveLinkOpenTrackIOSource::HandleInboundData);
 	UnicastReceiver->SetMaxReadBufferSize(2048);
 	UnicastReceiver->Start();
@@ -628,11 +646,6 @@ bool FLiveLinkOpenTrackIOSource::OpenUnicastSocket()
 
 bool FLiveLinkOpenTrackIOSource::OpenSockets()
 {
-	if (ConnectionState != ELiveLinkOpenTrackIOState::EndpointsReady)
-	{
-		return false;
-	}
-
 	switch (ConnectionSettings.Protocol)
 	{
 	case ELiveLinkOpenTrackIONetworkProtocol::Unicast:
@@ -658,63 +671,8 @@ bool FLiveLinkOpenTrackIOSource::OpenSockets()
 		return false;
 	}
 	}
-
-	ConnectionState = ELiveLinkOpenTrackIOState::Receiving;
 	
 	return true;
 }
 
-uint32 FLiveLinkOpenTrackIOSource::Run()
-{
-	double SleepDeltaTime = DevicePollWaitTime();
-
-	double LastDataReadTime = FPlatformTime::Seconds();
-	while (!Stopping)
-	{
-		// Make sure we do not exceed our local update rate.
-		ON_SCOPE_EXIT{
-			FPlatformProcess::Sleep(SleepDeltaTime);
-		};
-
-		if (!GIsRunning)
-		{
-			continue;
-		}
-
-		if (!bHasStartedSockets)
-		{
-			bHasStartedSockets = OpenSockets();
-			continue;
-		}
-
-		/** Read the incoming data queue and push that into Live Link. */
-		if (ConsumeInboundData())
-		{
-			LastDataReadTime = FPlatformTime::Seconds();
-		}
-
-		/** If the user has requested a connection reset. Close the socket and attempt to re-open on the next loop iteration. */
-		if (ConnectionState == ELiveLinkOpenTrackIOState::ResetRequested)
-		{
-			CloseSockets();
-
-			/** Re-parse the endpoints to reset the connection. This will set the flag to allow the socket to reconnect. */
-			ParseEndpointFromSourceSettings();
-			bHasStartedSockets = false;
-		}
-
-		if ((FPlatformTime::Seconds() - LastDataReadTime) > 1.0)
-		{
-			/** We have lost in-bound data.  Reset the socket read count. */
-			OpenTrackIOPacketCount = 0;
-		}
-
-		SleepDeltaTime = DevicePollWaitTime();
-	}
-
-	CloseSockets();
-	
-	return 0;
-}
-
 #undef LOCTEXT_NAMESPACE
diff --git a/Engine/Plugins/VirtualProduction/LiveLinkOpenTrackIO/Source/LiveLinkOpenTrackIO/Private/LiveLinkOpenTrackIOSource.h b/Engine/Plugins/VirtualProduction/LiveLinkOpenTrackIO/Source/LiveLinkOpenTrackIO/Private/LiveLinkOpenTrackIOSource.h
--- a/Engine/Plugins/VirtualProduction/LiveLinkOpenTrackIO/Source/LiveLinkOpenTrackIO/Private/LiveLinkOpenTrackIOSource.h
+++ b/Engine/Plugins/VirtualProduction/LiveLinkOpenTrackIO/Source/LiveLinkOpenTrackIO/Private/LiveLinkOpenTrackIOSource.h
@@ -7,11 +7,10 @@
 #include "Interfaces/IPv4/IPv4Endpoint.h"
 
 #include "LiveLinkOpenTrackIOConnectionSettings.h"
-#include "LiveLinkOpenTrackIOSourceSettings.h"
-#include "LiveLinkOpenTrackIOUdpSocketReceiver.h"
-
 #include "LiveLinkOpenTrackIOParser.h"
+#include "LiveLinkOpenTrackIOSourceSettings.h"
 
+#include "Common/UdpSocketReceiver.h"
 #include "Containers/MpscQueue.h"
 #include "Templates/UniquePtr.h"
 #include "Templates/PimplPtr.h"
@@ -44,10 +43,13 @@ enum class ELiveLinkOpenTrackIOState : uint8
 	Receiving,
 
 	/** The user has changed endpoints or user requested a socket reset. */
-	ResetRequested
+	ResetRequested,
+
+	/** Final state */
+	ShutDown,
 };
 
-class LIVELINKOPENTRACKIO_API FLiveLinkOpenTrackIOSource : public ILiveLinkSource, public FRunnable, public TSharedFromThis<FLiveLinkOpenTrackIOSource>
+class LIVELINKOPENTRACKIO_API FLiveLinkOpenTrackIOSource : public ILiveLinkSource, public TSharedFromThis<FLiveLinkOpenTrackIOSource>
 {
 public:
 
@@ -73,14 +75,6 @@ public:
 	virtual TSubclassOf<ULiveLinkSourceSettings> GetSettingsClass() const override { return ULiveLinkOpenTrackIOSourceSettings::StaticClass(); }
 	// End ILiveLinkSource Interface
 
-	// Begin FRunnable Interface
-	virtual bool Init() override { return true; }
-	virtual uint32 Run() override;
-	void Start();
-	virtual void Stop() override;
-	virtual void Exit() override { }
-	// End FRunnable Interface
-
 private:
 
 	/** Create listening sockets for OpenTrack channels. Returns true if socket handles were successfully opened and receivers started. */
@@ -95,31 +89,39 @@ private:
 	/** Close any open sockets. */
 	void CloseSockets();
 
-	/** Time in seconds to wait before polling the device values again. */
-	double DevicePollWaitTime() const;
-
-	/** Read the source settings and assign the endpoint addresses. */
-	void ParseEndpointFromSourceSettings();
+	/** Stop the udp receivers, but does not destroy the sockets. */
+	void StopUdpReceivers();
 
-	/** Trigger the reset mechanism. New sockets will be recreated. */
-	void ResetSockets();
+	/** Read the source settings and assign the endpoint addresses. Returns true if the endpoints are ready */
+	bool ParseEndpointFromSourceSettings();
 
 	/** Delegate for handling inbound segments. */
 	void HandleInboundData(const TSharedPtr<FArrayReader, ESPMode::ThreadSafe>& Data, const FIPv4Endpoint& InSender);
 
 	/** Consume inbound data. Fetch any data packets off of the inbound socket queue and process them. Returns true if data was processed. */
-	bool ConsumeInboundData();
+	void ConsumeInboundData();
 
 	/** Copy the open track data into Live Link equivalent data. This will establish any static data if the InData specifies it. */
 	void PushDataToLiveLink_AnyThread(const FLiveLinkOpenTrackIODatagramHeader& Header, FLiveLinkOpenTrackIOData InData);
 
-private:
 	/** Remove all transform subjects from the LL source. */
 	void RemoveAllTransformSubjects(FLiveLinkOpenTrackIOCache* Cache);
 
 	/** Push transform data to Live Link if the user has asked for it in the SavedSourceSettings. New subjects will get created automatically. */
 	void ConditionallyPushLiveLinkTransformData(FLiveLinkOpenTrackIOCache* Cache, const FLiveLinkOpenTrackIOData& InData);
 	
+	/** Sets connection state */
+	inline void SetConnectionState(const ELiveLinkOpenTrackIOState InConnectionState)
+	{
+		checkSlow(IsInGameThread());
+		ConnectionState = InConnectionState;
+	}
+
+	/** Called when FCoreDelegates::OnEnginePreExit fires */
+	void OnEnginePreExit();
+
+private:
+
 	struct FInboundMessage
 	{
 		/** Holds the segment data. */
@@ -162,10 +164,10 @@ private:
 	FSocket* MulticastSocket = nullptr;
 
 	/** Holds the unicast socket receiver. */
-	TUniquePtr<FLiveLinkOpenTrackIOUdpSocketReceiver> UnicastReceiver;
+	TUniquePtr<FUdpSocketReceiver> UnicastReceiver;
 	
 	/** Holds the multicast socket receiver. */
-	TUniquePtr<FLiveLinkOpenTrackIOUdpSocketReceiver> MulticastReceiver;
+	TUniquePtr<FUdpSocketReceiver> MulticastReceiver;
 
 	/** Holds the unicast socket. */
 	FSocket* UnicastSocket = nullptr;
@@ -175,34 +177,19 @@ private:
 
 	/** Multicast Endpoint endpoint to receive messages. */
 	FIPv4Endpoint MulticastEndpoint;
-	
-	/** Threadsafe Bool for terminating the main thread loop */
-	FThreadSafeBool Stopping;
-	
-	/** Thread to run updates on */
-	TUniquePtr<FRunnableThread> Thread;
-
-	/** Name of the running updates thread. Stored for debuggablity across multiple sources. */
-	FString ThreadName;
-
-	/** We don't start immediately after source contruction. We wait for the first engine tick. */
-	FDelegateHandle DeferredStartDelegateHandle;
 
-	/** Rate at which we poll the inbound queues. */
-	uint32 LocalUpdateRateInHz = 60;
+	/** State of the state machine. Only write/read in the game thread */
+	ELiveLinkOpenTrackIOState ConnectionState = ELiveLinkOpenTrackIOState::NotStarted;
 
-	/** Bool to indicate if we have started listening connections. Only used in running thread. */
-	bool bHasStartedSockets = false;
-
-	/** Variable to track the connection state from the game thread so that we can update the source status text. */
-	ELiveLinkOpenTrackIOState GameThreadConnectionState = ELiveLinkOpenTrackIOState::NotStarted;
+	/** Cached sender ip address, used to update the UI */
+	FIPv4Endpoint LastSender_RunnableThread;
 
-	/** Connection status */
-	std::atomic<ELiveLinkOpenTrackIOState> ConnectionState = ELiveLinkOpenTrackIOState::NotStarted;
+	/** Used to detect staleness in the source */
+	std::atomic<double> LastDataReadTime = 0;
 
-	/** Counter to track the number of inbound OpenTrackIO packets received. */
-	std::atomic<uint64> OpenTrackIOPacketCount = 0;
+	/** Tells the state machine that a shutdown of the source has been requested. One time only use, will never get cleared. */
+	bool bShutdownRequested = false;
 
-	/** Cached sender ip address, used to update the UI */
-	FIPv4Endpoint LastSender_RunnableThread;
+	/** Tells the state machine that the connection should be reset, probably because a relevant setting has been changed. */
+	bool bResetRequested = false;
 };
diff --git a/Engine/Plugins/VirtualProduction/LiveLinkOpenTrackIO/Source/LiveLinkOpenTrackIO/Private/LiveLinkOpenTrackIOUdpSocketReceiver.h b/Engine/Plugins/VirtualProduction/LiveLinkOpenTrackIO/Source/LiveLinkOpenTrackIO/Private/LiveLinkOpenTrackIOUdpSocketReceiver.h
deleted file mode 100644
diff --git a/Engine/Plugins/VirtualProduction/LiveLinkOpenTrackIO/Source/LiveLinkOpenTrackIO/Public/LiveLinkOpenTrackIOSourceSettings.h b/Engine/Plugins/VirtualProduction/LiveLinkOpenTrackIO/Source/LiveLinkOpenTrackIO/Public/LiveLinkOpenTrackIOSourceSettings.h
--- a/Engine/Plugins/VirtualProduction/LiveLinkOpenTrackIO/Source/LiveLinkOpenTrackIO/Public/LiveLinkOpenTrackIOSourceSettings.h
+++ b/Engine/Plugins/VirtualProduction/LiveLinkOpenTrackIO/Source/LiveLinkOpenTrackIO/Public/LiveLinkOpenTrackIOSourceSettings.h
@@ -49,22 +49,6 @@ public:
 	UPROPERTY(EditAnywhere, Category=Transport)
 	FString UnicastEndpoint = TEXT("0.0.0.0:0");
 
-	/**
-	 * The assigned IP address and port used by the socket layer.  This is for inspection of the connection
-	 * setup by the source via either blueprint or details panels for the source. 
-	 *
-	 * If this value starts with 0.0.0.0 then any adapter can be used with the given port number.
-	 * The format is IP_ADDRESS:PORT_NUMBER.
-	 *
-	 * If no receiver is setup then this value will be an empty string. 
-	 *
-	 */
-	UPROPERTY(VisibleAnywhere, Category=Transport, meta = (
-		EditCondition = "Protocol == ELiveLinkOpenTrackIONetworkProtocol::Unicast",
-		EditConditionHides
-		))
-	FString AssignedUnicastEndpoint = TEXT("");
-
 	/**
 	 * Emit subjects for each transform included in the OpenTrackIO channel.
 	 */
