Udp Socket Receiver: Add a way to specify the thead stack size.

Defaults to the original hard-coded value of 128*1024.

#rnx

diff --git a/Engine/Plugins/VirtualProduction/LiveLinkOpenTrackIO/Source/LiveLinkOpenTrackIO/Private/LiveLinkOpenTrackIOSource.cpp b/Engine/Plugins/VirtualProduction/LiveLinkOpenTrackIO/Source/LiveLinkOpenTrackIO/Private/LiveLinkOpenTrackIOSource.cpp
--- a/Engine/Plugins/VirtualProduction/LiveLinkOpenTrackIO/Source/LiveLinkOpenTrackIO/Private/LiveLinkOpenTrackIOSource.cpp
+++ b/Engine/Plugins/VirtualProduction/LiveLinkOpenTrackIO/Source/LiveLinkOpenTrackIO/Private/LiveLinkOpenTrackIOSource.cpp
@@ -605,6 +605,7 @@ bool FLiveLinkOpenTrackIOSource::OpenMulticastSocket()
 	MulticastReceiver = MakeUnique<FUdpSocketReceiver>(MulticastSocket, ThreadWaitTime, TEXT("LiveLinkOpenTransportIO_MulticastReceiver"));
 	MulticastReceiver->OnDataReceived().BindRaw(this, &FLiveLinkOpenTrackIOSource::HandleInboundData);
 	MulticastReceiver->SetMaxReadBufferSize(2048);
+	MulticastReceiver->SetThreadStackSize(512 * 1024);
 	MulticastReceiver->Start();
 
 	return true;
@@ -638,6 +639,7 @@ bool FLiveLinkOpenTrackIOSource::OpenUnicastSocket()
 	UnicastReceiver = MakeUnique<FUdpSocketReceiver>(UnicastSocket, ThreadWaitTime, TEXT("LiveLinkOpenTransportIO_UnicastReceiver"));
 	UnicastReceiver->OnDataReceived().BindRaw(this, &FLiveLinkOpenTrackIOSource::HandleInboundData);
 	UnicastReceiver->SetMaxReadBufferSize(2048);
+	UnicastReceiver->SetThreadStackSize(512 * 1024);
 	UnicastReceiver->Start();
 
 	return true;
diff --git a/Engine/Source/Runtime/Networking/Public/Common/UdpSocketReceiver.h b/Engine/Source/Runtime/Networking/Public/Common/UdpSocketReceiver.h
--- a/Engine/Source/Runtime/Networking/Public/Common/UdpSocketReceiver.h
+++ b/Engine/Source/Runtime/Networking/Public/Common/UdpSocketReceiver.h
@@ -76,10 +76,16 @@ public:
 		MaxReadBufferSize = InMaxReadBufferSize;
 	}
 
+	/** Set the desired Thread Stack Size. Only read when Start() is called. */
+	void SetThreadStackSize(uint32 InTheadStackSize)
+	{
+		ThreadStackSize = InTheadStackSize;
+	}
+
 	/** Start the receiver thread. */
 	void Start()
 	{
-		Thread = FRunnableThread::Create(this, *ThreadName, 128 * 1024, TPri_AboveNormal, FPlatformAffinity::GetPoolThreadMask());
+		Thread = FRunnableThread::Create(this, *ThreadName, ThreadStackSize, TPri_AboveNormal, FPlatformAffinity::GetPoolThreadMask());
 	}
 
 	/**
@@ -189,6 +195,9 @@ private:
 	/** The maximum read buffer size used to read the socket. */
 	uint32 MaxReadBufferSize = 65507u;
 
+	/** Thread stack size to use. */
+	uint32 ThreadStackSize = 128 * 1024;
+
 private:
 
 	/** Holds the data received delegate. */
