OpenTrackIO: Added USTRUCT based replacements for TOptional values that are to be exposed in Blueprints.

The Cbor and Json parsers are now aware of these types and populate them accordingly.

#rnx

diff --git a/Engine/Plugins/VirtualProduction/LiveLinkOpenTrackIO/Source/LiveLinkOpenTrackIO/Private/LiveLinkOpenTrackIOParser.cpp b/Engine/Plugins/VirtualProduction/LiveLinkOpenTrackIO/Source/LiveLinkOpenTrackIO/Private/LiveLinkOpenTrackIOParser.cpp
--- a/Engine/Plugins/VirtualProduction/LiveLinkOpenTrackIO/Source/LiveLinkOpenTrackIO/Private/LiveLinkOpenTrackIOParser.cpp
+++ b/Engine/Plugins/VirtualProduction/LiveLinkOpenTrackIO/Source/LiveLinkOpenTrackIO/Private/LiveLinkOpenTrackIOParser.cpp
@@ -94,20 +94,118 @@ namespace UE::OpenTrackIO::Private
 {
 
 /**
-* Callback for importing JSON values into TOptional (currently numeric) properties.
-* If the property is not an FOptionalProperty or does not wrap a numeric type, 
-* the callback returns false to let the default parser take over.
-*
-* @param JsonValue   The JSON value to import.
-* @param Property    The TOptional property to write into.
-* @param Value       Pointer to memory holding the optional property value.
-* 
-* @return            True if this callback handled the import (optional set successfully).
-*/
-bool JsonOptionalImporterCallback(const TSharedPtr<FJsonValue>& JsonValue, FProperty* Property, void* Value)
+ * Tries to parse our custom Optional USTRUCT from JSON.
+ * 
+ * @param JsonValue   The JSON value to import.
+ * @param Property    The TOptional property to write into.
+ * @param Value       Pointer to memory holding the optional property value.
+ * 
+ * @return true on success.
+ */
+static bool TryReadStructOptional(
+	const TSharedPtr<FJsonValue>& JsonValue,
+	FProperty* Property,
+	void* ContainerPtr)
 {
-	// Only handle optionals
+	// See if this is a struct property
+
+	FStructProperty* StructProperty = CastField<FStructProperty>(Property);
+
+	if (!StructProperty)
+	{
+		return false;
+	}
+
+	// Make sure it is one of our toptional replacement structs.
+
+	UScriptStruct* ScriptStruct = StructProperty->Struct;
+
+	if (!IsOpenTrackIOOptionalType(ScriptStruct))
+	{
+		return false;
+	}
+
+	// Find the 'Value' member
+
+	FProperty* ValueProperty = ScriptStruct->FindPropertyByName(OptionalTypeValueName);
+
+	if (!ValueProperty)
+	{
+		return false;
+	}
+
+	// Only handle numeric inner properties
+
+	FNumericProperty* NumericProperty = CastField<FNumericProperty>(ValueProperty);
+
+	if (!NumericProperty)
+	{
+		return false;
+	}
+
+	// Get pointer to the numeric field inside the struct
+
+	if (!ContainerPtr)
+	{
+		return false;
+	}
+
+	void* NumericPtr = NumericProperty->ContainerPtrToValuePtr<void>(ContainerPtr);
+
+	if (!NumericPtr)
+	{
+		return false;
+	}
+
+	// Set numeric value
+
+	if (NumericProperty->IsFloatingPoint())
+	{
+		NumericProperty->SetFloatingPointPropertyValue(NumericPtr, JsonValue->AsNumber());
+	}
+	else if (NumericProperty->IsInteger())
+	{
+		if (JsonValue->Type == EJson::String)
+		{
+			// parse string -> int64 ourselves so we don't lose any precision going through AsNumber (aka double)
+			NumericProperty->SetIntPropertyValue(NumericPtr, FCString::Atoi64(*JsonValue->AsString()));
+		}
+		else
+		{
+			NumericProperty->SetIntPropertyValue(NumericPtr, (int64)JsonValue->AsNumber());
+		}
+	}
+	else
+	{
+		return false;
+	}
+
+	// Mark bIsSet = true
+
+	FBoolProperty* BoolProp = CastField<FBoolProperty>(ScriptStruct->FindPropertyByName(OptionalTypeIsSetName));
+
+	if (BoolProp)
+	{
+		BoolProp->SetPropertyValue_InContainer(ContainerPtr, true);
+	}
+
+	return true;
+}
 
+/**
+ * Tries to parse a TOptional numeric property from JSON.
+ * 
+ * @param JsonValue   The JSON value to import.
+ * @param Property    The TOptional property to write into.
+ * @param Value       Pointer to memory holding the optional property value.
+ * 
+ * @return true if successful.
+ */
+static bool TryReadTOptional(
+	const TSharedPtr<FJsonValue>& JsonValue,
+	FProperty* Property,
+	void* Value)
+{
 	FOptionalProperty* OptionalProperty = CastField<FOptionalProperty>(Property);
 
 	if (!OptionalProperty)
@@ -139,7 +237,7 @@ bool JsonOptionalImporterCallback(const TSharedPtr<FJsonValue>& JsonValue, FProp
 
 		return true;
 	}
-		
+
 	// Handle integer property types.
 
 	if (NumericProperty->IsInteger())
@@ -168,6 +266,28 @@ bool JsonOptionalImporterCallback(const TSharedPtr<FJsonValue>& JsonValue, FProp
 	return false;
 }
 
+/**
+* Callback for importing JSON values into TOptional (currently numeric) properties,
+* or our custom OptionalStruct types that can be natively exposed to Blueprints.
+*
+* @param JsonValue   The JSON value to import.
+* @param Property    The TOptional property to write into.
+* @param Value       Pointer to memory holding the optional property value.
+* 
+* @return            True if this callback handled the import (optional set successfully).
+*/
+static bool JsonOptionalImporterCallback(const TSharedPtr<FJsonValue>& JsonValue, FProperty* Property, void* Value)
+{
+	// Try struct-based optionals
+	if (TryReadStructOptional(JsonValue, Property, Value))
+	{
+		return true;
+	}
+
+	// Try TOptional
+	return TryReadTOptional(JsonValue, Property, Value);
+}
+
 TOptional<FLiveLinkOpenTrackIOData> ParseJsonBlob(const FString& JsonBlob)
 {
 	TSharedPtr<FJsonObject> JsonObject;
@@ -231,7 +351,7 @@ TOptional<FLiveLinkOpenTrackIOData> ParseCborBlob(TArrayView<const uint8> InByte
  * OpenTrackIO specifies using Fletcher-16 algorithm for checksum. 
  * Optimized version from: https://en.wikipedia.org/wiki/Fletcher%27s_checksum
  */
-uint16_t CalculateChecksum(TArrayView<const uint8> HeaderBytes, TArrayView<const uint8> PayloadBytes)
+static uint16_t CalculateChecksum(TArrayView<const uint8> HeaderBytes, TArrayView<const uint8> PayloadBytes)
 {
 	uint32_t C0 = 0;
 	uint32_t C1 = 0;
@@ -315,7 +435,7 @@ inline FArchive& operator<<(FArchive& Ar, FLiveLinkOpenTrackIODatagramHeader& He
 	return Ar;
 }
 
-bool IsHeaderValid(TArrayView<const uint8> Bytes)
+static bool IsHeaderValid(TArrayView<const uint8> Bytes)
 {
 	const uint32 HeaderSize = sizeof(FLiveLinkOpenTrackIODatagramHeader);
 	/** No way this is a valid datagram if we can't fit it in our header. */
@@ -343,7 +463,7 @@ bool IsHeaderValid(TArrayView<const uint8> Bytes)
 	return true;
 }
 
-TOptional<FLiveLinkOpenTrackIOData> ConvertPayloadToData(const FLiveLinkOpenTrackIODatagramHeader& Header, TArrayView<const uint8> Payload)
+static TOptional<FLiveLinkOpenTrackIOData> ConvertPayloadToData(const FLiveLinkOpenTrackIODatagramHeader& Header, TArrayView<const uint8> Payload)
 {
 	/** Now attempt to parse the payload using either JSON or CBOR. */
 	TOptional<FLiveLinkOpenTrackIOData> Data;
@@ -370,7 +490,7 @@ TOptional<FLiveLinkOpenTrackIOData> ConvertPayloadToData(const FLiveLinkOpenTrac
  *
  * If this fails then the optional will be unset. 
  */
-TOptional<FLiveLinkOpenTrackIOData> TryEmitOpenTrackData(FLiveLinkOpenTrackIODatagramHeader Header, TArrayView<const uint8> Payload)
+static TOptional<FLiveLinkOpenTrackIOData> TryEmitOpenTrackData(FLiveLinkOpenTrackIODatagramHeader Header, TArrayView<const uint8> Payload)
 {
 	TOptional<FLiveLinkOpenTrackIOData> Data = ConvertPayloadToData(Header, Payload);
 	if (Data && Data->Protocol.IsSupported())
diff --git a/Engine/Plugins/VirtualProduction/LiveLinkOpenTrackIO/Source/LiveLinkOpenTrackIO/Private/LiveLinkOpenTrackIOSource.cpp b/Engine/Plugins/VirtualProduction/LiveLinkOpenTrackIO/Source/LiveLinkOpenTrackIO/Private/LiveLinkOpenTrackIOSource.cpp
--- a/Engine/Plugins/VirtualProduction/LiveLinkOpenTrackIO/Source/LiveLinkOpenTrackIO/Private/LiveLinkOpenTrackIOSource.cpp
+++ b/Engine/Plugins/VirtualProduction/LiveLinkOpenTrackIO/Source/LiveLinkOpenTrackIO/Private/LiveLinkOpenTrackIOSource.cpp
@@ -334,61 +334,52 @@ void FLiveLinkOpenTrackIOSource::OnSettingsChanged(ULiveLinkSourceSettings* Sett
 	}
 }
 	
-void FLiveLinkOpenTrackIOSource::HandleInboundData(const TSharedPtr<FArrayReader, ESPMode::ThreadSafe>& Data, const FIPv4Endpoint& InSender)
+void FLiveLinkOpenTrackIOSource::HandleInboundData(const TSharedPtr<FArrayReader, ESPMode::ThreadSafe>& InData, const FIPv4Endpoint& InSender)
 {
-	Inbound.Enqueue(FInboundMessage{ Data, InSender });
-	ConsumeInboundData();
-}
+	TRACE_CPUPROFILER_EVENT_SCOPE(OpenTrackIO::HandleInboundData);
 
-void FLiveLinkOpenTrackIOSource::ConsumeInboundData()
-{
-	TRACE_CPUPROFILER_EVENT_SCOPE(OpenTrackIO::ConsumeInboundData);
+	FInboundMessage Message{ InData, InSender };
+
+	if (!Message.Data.IsValid())
+	{
+		return;
+	}
 
-	FInboundMessage Message;
+	TArrayView<const uint8> MessageData(Message.Data->GetData(), Message.Data->Num());
 
-	while (Inbound.Dequeue(Message))
+	FOpenTrackIOHeaderWithPayload& PayloadContainer = WorkingPayloads.FindOrAdd(Message.Sender);
+	if (UE::OpenTrackIO::Private::GetHeaderAndPayloadFromBytes(MessageData, PayloadContainer))
 	{
-		if (!Message.Data.IsValid())
+		// Update Source Machine so the user can identify the sending machine.
+		if (LastSender_RunnableThread != Message.Sender)
 		{
-			continue;
+			LastSender_RunnableThread = Message.Sender;
+			AsyncTask(ENamedThreads::GameThread, [this, Sender = Message.Sender]() {
+					SourceMachineName = Sender.Address.ToText();
+			});
 		}
 
-		TArrayView<const uint8> MessageData(Message.Data->GetData(), Message.Data->Num());
+		LastDataReadTime = FPlatformTime::Seconds();
 
-		FOpenTrackIOHeaderWithPayload& PayloadContainer = WorkingPayloads.FindOrAdd(Message.Sender);
-		if (UE::OpenTrackIO::Private::GetHeaderAndPayloadFromBytes(MessageData, PayloadContainer))
+		if (PayloadContainer.GetHeader().IsLastSegment())
 		{
-			// Update Source Machine so the user can identify the sending machine.
-			if (LastSender_RunnableThread != Message.Sender)
-			{
-				LastSender_RunnableThread = Message.Sender;
-				AsyncTask(ENamedThreads::GameThread, [this, Sender = Message.Sender]() {
-						SourceMachineName = Sender.Address.ToText();
-				});
-			}
-
-			LastDataReadTime = FPlatformTime::Seconds();
-
-			if (PayloadContainer.GetHeader().IsLastSegment())
 			{
+				// Scoped because want to prevent access to const & Header beyond this use.
+				const FLiveLinkOpenTrackIODatagramHeader& Header = PayloadContainer.GetHeader();
+				TOptional<FLiveLinkOpenTrackIOData> Data = UE::OpenTrackIO::Private::ParsePayload(PayloadContainer);
+				if (Data)
 				{
-					// Scoped because want to prevent access to const & Header beyond this use.
-					const FLiveLinkOpenTrackIODatagramHeader& Header = PayloadContainer.GetHeader();
-					TOptional<FLiveLinkOpenTrackIOData> Data = UE::OpenTrackIO::Private::ParsePayload(PayloadContainer);
-					if (Data)
-					{
-						// This is a sink and we give up FLiveLinkOpenTrackData to optimize moving the OpenTrack data around.
-						PushDataToLiveLink_AnyThread(Header, MoveTemp(*Data));
-					}
+					// This is a sink and we give up FLiveLinkOpenTrackData to optimize moving the OpenTrack data around.
+					PushDataToLiveLink_AnyThread(Header, MoveTemp(*Data));
 				}
-				// Remove the working payload from our table. This must happen 
-				WorkingPayloads.Remove(Message.Sender);
 			}
+			// Remove the working payload from our table.
+			WorkingPayloads.Remove(Message.Sender);
 		}
-		else
-		{
-			UE_LOG(LogLiveLinkOpenTrackIO, Display, TEXT("Failed to handle inbound data."));
-		}
+	}
+	else
+	{
+		UE_LOG(LogLiveLinkOpenTrackIO, Display, TEXT("Failed to handle inbound data."));
 	}
 }
 
diff --git a/Engine/Plugins/VirtualProduction/LiveLinkOpenTrackIO/Source/LiveLinkOpenTrackIO/Private/LiveLinkOpenTrackIOSource.h b/Engine/Plugins/VirtualProduction/LiveLinkOpenTrackIO/Source/LiveLinkOpenTrackIO/Private/LiveLinkOpenTrackIOSource.h
--- a/Engine/Plugins/VirtualProduction/LiveLinkOpenTrackIO/Source/LiveLinkOpenTrackIO/Private/LiveLinkOpenTrackIOSource.h
+++ b/Engine/Plugins/VirtualProduction/LiveLinkOpenTrackIO/Source/LiveLinkOpenTrackIO/Private/LiveLinkOpenTrackIOSource.h
@@ -96,10 +96,7 @@ private:
 	bool ParseEndpointFromSourceSettings();
 
 	/** Delegate for handling inbound segments. */
-	void HandleInboundData(const TSharedPtr<FArrayReader, ESPMode::ThreadSafe>& Data, const FIPv4Endpoint& InSender);
-
-	/** Consume inbound data. Fetch any data packets off of the inbound socket queue and process them. Returns true if data was processed. */
-	void ConsumeInboundData();
+	void HandleInboundData(const TSharedPtr<FArrayReader, ESPMode::ThreadSafe>& InData, const FIPv4Endpoint& InSender);
 
 	/** Copy the open track data into Live Link equivalent data. This will establish any static data if the InData specifies it. */
 	void PushDataToLiveLink_AnyThread(const FLiveLinkOpenTrackIODatagramHeader& Header, FLiveLinkOpenTrackIOData InData);
@@ -140,9 +137,6 @@ private:
 	/** Cache of data relevant to open track protocol. It is keyed by the sourceId. We can support multiple source ids */
 	TMap<FString, FOpenTrackIOCachePtr> OpenTrackIOCacheMap;
 	
-	/** Holds the queue of inbound messages. */
-	TMpscQueue<FInboundMessage> Inbound;
-
 	/** The LiveLink client provided via ReceiveClient*/
 	ILiveLinkClient* Client = nullptr;
 
diff --git a/Engine/Plugins/VirtualProduction/LiveLinkOpenTrackIO/Source/LiveLinkOpenTrackIO/Private/LiveLinkOpenTrackIOTypes.h b/Engine/Plugins/VirtualProduction/LiveLinkOpenTrackIO/Source/LiveLinkOpenTrackIO/Private/LiveLinkOpenTrackIOTypes.h
--- a/Engine/Plugins/VirtualProduction/LiveLinkOpenTrackIO/Source/LiveLinkOpenTrackIO/Private/LiveLinkOpenTrackIOTypes.h
+++ b/Engine/Plugins/VirtualProduction/LiveLinkOpenTrackIO/Source/LiveLinkOpenTrackIO/Private/LiveLinkOpenTrackIOTypes.h
@@ -12,11 +12,221 @@
 
 #include "LiveLinkOpenTrackIOTypes.generated.h"
 
-namespace UE::OpenTrackIO
+
+/**
+ * Contains a float and bIsSet flag, mimicking TOptional, to be exposed in Blueprints.
+ * The intention is that it is clear to the user if the value contained has been set
+ * by the sender or not.
+ */
+USTRUCT(BlueprintType, meta = (DisplayName = "Optional Float (OpenTrackIO)"))
+struct FOpenTrackIOOptionalFloat
 {
-	/** Brown-Conrady D-U, the default for nameless distortion model per opentrackio.org */
-	static const FName BrownConradyDU = TEXT("Brown-Conrady D-U");
-}
+	GENERATED_BODY()
+
+	/** Whether Value has been set or not */
+	UPROPERTY(EditAnywhere, BlueprintReadWrite, Category = "Optional")
+	bool bIsSet = false;
+
+	/** The actual float value (only valid when bIsSet == true) */
+	UPROPERTY(EditAnywhere, BlueprintReadWrite, Category = "Optional", meta = (EditCondition = "bIsSet"))
+	float Value = 0;
+
+public:
+
+	/** Equals Operator */
+	bool operator==(const FOpenTrackIOOptionalFloat& Other) const
+	{
+		if (bIsSet != Other.bIsSet)
+		{
+			return false;
+		}
+
+		// If neither is set, we consider them equal even if the Value is different.
+		if (!bIsSet)
+		{
+			return true;
+		}
+
+		// both are set, we compare the Value
+		return FMath::IsNearlyEqual(Value, Other.Value);
+	}
+
+	/** Different Than operator */
+	bool operator!=(const FOpenTrackIOOptionalFloat& Other) const
+	{
+		return !(*this == Other);
+	}
+
+	/** Check if Value has been set. Mimicks IsSet() in TOptional */
+	bool IsSet() const
+	{
+		return bIsSet;
+	}
+
+	/** Gets the set Value. Mimicks GetValue() in TOptional */
+	float GetValue() const
+	{
+		return Value;
+	}
+
+	/** Sets the Value of this container */
+	void SetValue(float InValue)
+	{
+		bIsSet = true;
+		Value = InValue;
+	}
+
+	/** Clears the bIsSet flag. Mimicks Reset() in TOptional */
+	void Reset()
+	{
+		bIsSet = false;
+		Value = 0.0f;
+	}
+};
+
+
+/**
+ * Contains an int32 and bIsSet flag, mimicking TOptional, to be exposed in Blueprints.
+ * The intention is that it is clear to the user if the value contained has been set
+ * by the sender or not.
+ */
+USTRUCT(BlueprintType, meta = (DisplayName = "Optional Int32 (OpenTrackIO)"))
+struct FOpenTrackIOOptionalInt32
+{
+	GENERATED_BODY()
+
+	/** Whether Value has been set or not */
+	UPROPERTY(EditAnywhere, BlueprintReadWrite, Category = "Optional")
+	bool bIsSet = false;
+
+	/** The actual int32 value (only valid when bIsSet == true) */
+	UPROPERTY(EditAnywhere, BlueprintReadWrite, Category = "Optional", meta = (EditCondition = "bIsSet"))
+	int32 Value = 0;
+
+public:
+
+	/** Equals Operator */
+	bool operator==(const FOpenTrackIOOptionalInt32& Other) const
+	{
+		if (bIsSet != Other.bIsSet)
+		{
+			return false;
+		}
+
+		// If neither is set, we consider them equal even if the Value is different.
+		if (!bIsSet)
+		{
+			return true;
+		}
+
+		// both are set, we compare the Value
+		return Value == Other.Value;
+	}
+
+	/** Different Than operator */
+	bool operator!=(const FOpenTrackIOOptionalInt32& Other) const
+	{
+		return !(*this == Other);
+	}
+
+	/** Check if Value has been set. Mimicks IsSet() in TOptional */
+	bool IsSet() const
+	{
+		return bIsSet;
+	}
+
+	/** Gets the set Value. Mimicks GetValue() in TOptional */
+	int32 GetValue() const
+	{
+		return Value;
+	}
+
+	/** Sets the Value of this container */
+	void SetValue(int32 InValue)
+	{
+		bIsSet = true;
+		Value = InValue;
+	}
+
+	/** Clears the bIsSet flag. Mimicks Reset() in TOptional */
+	void Reset()
+	{
+		bIsSet = false;
+		Value = 0;
+	}
+};
+
+
+/**
+ * Contains a double and bIsSet flag, mimicking TOptional, to be exposed in Blueprints.
+ * The intention is that it is clear to the user if the value contained has been set
+ * by the sender or not.
+ */
+USTRUCT(BlueprintType, meta = (DisplayName = "Optional Double (OpenTrackIO)"))
+struct FOpenTrackIOOptionalDouble
+{
+	GENERATED_BODY()
+
+	/** Whether Value has been set or not */
+	UPROPERTY(EditAnywhere, BlueprintReadWrite, Category = "Optional")
+	bool bIsSet = false;
+
+	/** The actual double value (only valid when bIsSet == true) */
+	UPROPERTY(EditAnywhere, BlueprintReadWrite, Category = "Optional", meta = (EditCondition = "bIsSet"))
+	double Value = 0.0;
+
+public:
+
+	/** Equals Operator */
+	bool operator==(const FOpenTrackIOOptionalDouble& Other) const
+	{
+		if (bIsSet != Other.bIsSet)
+		{
+			return false;
+		}
+
+		// If neither is set, we consider them equal even if the Value is different.
+		if (!bIsSet)
+		{
+			return true;
+		}
+
+		// both are set, we compare the Value
+		return FMath::IsNearlyEqual(Value, Other.Value, UE_DOUBLE_SMALL_NUMBER);
+	}
+
+	/** Different Than operator */
+	bool operator!=(const FOpenTrackIOOptionalDouble& Other) const
+	{
+		return !(*this == Other);
+	}
+
+	/** Check if Value has been set. Mimicks IsSet() in TOptional */
+	bool IsSet() const
+	{
+		return bIsSet;
+	}
+
+	/** Gets the set Value. Mimicks GetValue() in TOptional */
+	double GetValue() const
+	{
+		return Value;
+	}
+
+	/** Sets the Value of this container */
+	void SetValue(double InValue)
+	{
+		bIsSet = true;
+		Value = InValue;
+	}
+
+	/** Clears the bIsSet flag. Mimicks Reset() in TOptional */
+	void Reset()
+	{
+		bIsSet = false;
+		Value = 0.0;
+	}
+};
 
 
 /**
@@ -112,10 +322,10 @@ struct FLiveLinkOpenTrackIOSensorDimensions
 	bool operator!=(const FLiveLinkOpenTrackIOSensorDimensions&) const = default;
 	
 	UPROPERTY(EditAnywhere, BlueprintReadWrite, Category = "OpenTrackIO")
-	TOptional<float> Height;
+	FOpenTrackIOOptionalFloat Height;
 
 	UPROPERTY(EditAnywhere, BlueprintReadWrite, Category = "OpenTrackIO")
-	TOptional<float> Width;
+	FOpenTrackIOOptionalFloat Width;
 };
 
 /**
@@ -1006,15 +1216,15 @@ struct FLiveLinkOpenTrackIOLens
 
 	/** The linear f-number of the lens, equal to the focal length divided by the diameter of the entrance pupil. */
 	UPROPERTY(EditAnywhere, BlueprintReadWrite, Category = "OpenTrackIO")
-	TOptional<float> FStop;
+	FOpenTrackIOOptionalFloat FStop;
 
 	/** Distance between the pinhole and the image plane in the simple CGI pinhole camera model. (millimeters) */
 	UPROPERTY(EditAnywhere, BlueprintReadWrite, Category = "OpenTrackIO")
-	TOptional<float> PinholeFocalLength;
+	FOpenTrackIOOptionalFloat PinholeFocalLength;
 
 	/** Focus distance/position of the lens in meters. */ 
 	UPROPERTY(EditAnywhere, BlueprintReadWrite, Category = "OpenTrackIO")
-	TOptional<float> FocusDistance;
+	FOpenTrackIOOptionalFloat FocusDistance;
 
 	/**  Offset in x and y of the centre of perspective projection of the virtual camera */
 	UPROPERTY(EditAnywhere, BlueprintReadWrite, Category = "OpenTrackIO")
@@ -1026,7 +1236,7 @@ struct FLiveLinkOpenTrackIOLens
 
 	/** T-stop: Linear t-number of the lens, equal to the F-number of the lens divided by the square root of the transmittance of the lens. */
 	UPROPERTY(EditAnywhere, BlueprintReadWrite, Category = "OpenTrackIO")
-	float TStop = 0.0;
+	FOpenTrackIOOptionalFloat TStop;
 };
 
 /**
@@ -1364,3 +1574,37 @@ struct FLiveLinkOpenTrackIOData
 	FLiveLinkOpenTrackIOCustomData Custom;
 };
 
+
+namespace UE::OpenTrackIO
+{
+	/** Brown-Conrady D-U, the default for nameless distortion model per opentrackio.org */
+	static const FName BrownConradyDU = TEXT("Brown-Conrady D-U");
+
+	/** Name of the FOpenTrackIOOptional "bIsSet" property */
+	static const FName OptionalTypeIsSetName = GET_MEMBER_NAME_CHECKED(FOpenTrackIOOptionalFloat, bIsSet);
+
+	/** Name of the FOpenTrackIOOptional "Value" property */
+	static const FName OptionalTypeValueName = GET_MEMBER_NAME_CHECKED(FOpenTrackIOOptionalFloat, Value);
+
+	/**
+	 * Returns true if Struct is one of our custom optional USTRUCT types.
+	 */
+	static bool IsOpenTrackIOOptionalType(const UScriptStruct* ScriptStruct)
+	{
+		static UScriptStruct* OptionalStructs[] = {
+			FOpenTrackIOOptionalFloat::StaticStruct(),
+			FOpenTrackIOOptionalInt32::StaticStruct(),
+			FOpenTrackIOOptionalDouble::StaticStruct()
+		};
+
+		for (UScriptStruct* OptionalStruct : OptionalStructs)
+		{
+			if (ScriptStruct == OptionalStruct)
+			{
+				return true;
+			}
+		}
+
+		return false;
+	}
+}
diff --git a/Engine/Plugins/VirtualProduction/LiveLinkOpenTrackIO/Source/LiveLinkOpenTrackIO/Private/OpenTrackIOCborStructDeserializerBackend.cpp b/Engine/Plugins/VirtualProduction/LiveLinkOpenTrackIO/Source/LiveLinkOpenTrackIO/Private/OpenTrackIOCborStructDeserializerBackend.cpp
--- a/Engine/Plugins/VirtualProduction/LiveLinkOpenTrackIO/Source/LiveLinkOpenTrackIO/Private/OpenTrackIOCborStructDeserializerBackend.cpp
+++ b/Engine/Plugins/VirtualProduction/LiveLinkOpenTrackIO/Source/LiveLinkOpenTrackIO/Private/OpenTrackIOCborStructDeserializerBackend.cpp
@@ -3,6 +3,7 @@
 #include "OpenTrackIOCborStructDeserializerBackend.h"
 
 #include "LiveLinkOpenTrackIO.h"
+#include "LiveLinkOpenTrackIOLiveLinkTypes.h"
 
 #include "UObject/Class.h"
 #include "UObject/AnsiStrProperty.h"
@@ -13,7 +14,7 @@
 #include "UObject/Utf8StrProperty.h"
 
 
-namespace UE::OpenTrackIOCbor::Private
+namespace UE::OpenTrackIO::Cbor::Private
 {
 	/**
 	* Clears the value of the given property.
@@ -152,6 +153,41 @@ namespace UE::OpenTrackIOCbor::Private
 
 		return false;
 	}
+
+	/**
+	 * Attempts to read one of our optional USTRUCTs. Returns true on successful parse,
+	 * false otherwise (leaving bIsSet unchanged on failure).
+	 */
+	bool TryReadOurOptionalStruct(
+		FOpenTrackIOCborStructDeserializerBackend* Backend,
+		const UScriptStruct* OptionalStruct,
+		FStructProperty* StructProp,
+		FProperty* Outer,
+		void* Data,
+		int32 ArrayIndex)
+	{
+		void* StructData = GetPropertyValuePtr(StructProp, Outer, Data, ArrayIndex);
+		if (!StructData)
+		{
+			return false;
+		}
+
+		if (FProperty* ValueProp = OptionalStruct->FindPropertyByName(OptionalTypeValueName))
+		{
+			if (!Backend->ReadProperty(ValueProp, StructProp, StructData, /* Inner ArrayIndex */ 0))
+			{
+				return false;
+			}
+
+			if (FBoolProperty* BoolProp = CastField<FBoolProperty>(OptionalStruct->FindPropertyByName(OptionalTypeIsSetName)))
+			{
+				BoolProp->SetPropertyValue_InContainer(StructData, true);
+			}
+
+			return true;
+		}
+		return false;
+	}
 };
 
 
@@ -263,10 +299,23 @@ bool FOpenTrackIOCborStructDeserializerBackend::GetNextToken(EStructDeserializer
 
 bool FOpenTrackIOCborStructDeserializerBackend::ReadProperty(FProperty* Property, FProperty* Outer, void* Data, int32 ArrayIndex)
 {
+	// Support for our custom USTRUCT replacement of TOptionals, since TOptionals are not currently supported in blueprint,
+	// and we need that to effectively expose many of these optional parameters.
+	//
+	if (FStructProperty* StructProp = CastField<FStructProperty>(Property))
+	{
+		const UScriptStruct* UDStruct = StructProp->Struct;
+
+		if (UE::OpenTrackIO::IsOpenTrackIOOptionalType(UDStruct))
+		{
+			return UE::OpenTrackIO::Cbor::Private::TryReadOurOptionalStruct(this, UDStruct, StructProp, Outer, Data, ArrayIndex);
+		}
+	}
+
 	// Unwrap TOptional and recurse back into this function
 	if (FOptionalProperty* OptionalProperty = CastField<FOptionalProperty>(Property))
 	{
-		void* ValuePtr = UE::OpenTrackIOCbor::Private::GetPropertyValuePtr(OptionalProperty, Outer, Data, ArrayIndex);
+		void* ValuePtr = UE::OpenTrackIO::Cbor::Private::GetPropertyValuePtr(OptionalProperty, Outer, Data, ArrayIndex);
 
 		if (!ValuePtr)
 		{
@@ -283,6 +332,7 @@ bool FOpenTrackIOCborStructDeserializerBackend::ReadProperty(FProperty* Property
 		}
 
 		// Get the property inside the optional
+
 		FProperty* ValueProp = OptionalProperty->GetValueProperty();
 
 		if (!ValueProp)
@@ -307,7 +357,7 @@ bool FOpenTrackIOCborStructDeserializerBackend::ReadProperty(FProperty* Property
 	// Unsigned Integers
 	case ECborCode::Uint:
 	{
-		const bool bWasSet = UE::OpenTrackIOCbor::Private::SetNumericPropertyValue(Property, Outer, Data, ArrayIndex, LastContext.AsUInt());
+		const bool bWasSet = UE::OpenTrackIO::Cbor::Private::SetNumericPropertyValue(Property, Outer, Data, ArrayIndex, LastContext.AsUInt());
 
 		if (!bWasSet)
 		{
@@ -321,7 +371,7 @@ bool FOpenTrackIOCborStructDeserializerBackend::ReadProperty(FProperty* Property
 	// Signed Integers
 	case ECborCode::Int:
 	{
-		const bool bWasSet = UE::OpenTrackIOCbor::Private::SetNumericPropertyValue(Property, Outer, Data, ArrayIndex, LastContext.AsInt());
+		const bool bWasSet = UE::OpenTrackIO::Cbor::Private::SetNumericPropertyValue(Property, Outer, Data, ArrayIndex, LastContext.AsInt());
 
 		if (!bWasSet)
 		{
@@ -337,24 +387,24 @@ bool FOpenTrackIOCborStructDeserializerBackend::ReadProperty(FProperty* Property
 	{
 		if (FAnsiStrProperty* AnsiStrProperty = CastField<FAnsiStrProperty>(Property))
 		{
-			return UE::OpenTrackIOCbor::Private::SetPropertyValue(AnsiStrProperty, Outer, Data, ArrayIndex, LastContext.AsAnsiString());
+			return UE::OpenTrackIO::Cbor::Private::SetPropertyValue(AnsiStrProperty, Outer, Data, ArrayIndex, LastContext.AsAnsiString());
 		}
 
 		if (FUtf8StrProperty* Utf8StrProperty = CastField<FUtf8StrProperty>(Property))
 		{
-			return UE::OpenTrackIOCbor::Private::SetPropertyValue(Utf8StrProperty, Outer, Data, ArrayIndex, LastContext.AsUtf8String());
+			return UE::OpenTrackIO::Cbor::Private::SetPropertyValue(Utf8StrProperty, Outer, Data, ArrayIndex, LastContext.AsUtf8String());
 		}
 
 		const FString StringValue = LastContext.AsString();
 
 		if (FStrProperty* StrProperty = CastField<FStrProperty>(Property))
 		{
-			return UE::OpenTrackIOCbor::Private::SetPropertyValue(StrProperty, Outer, Data, ArrayIndex, StringValue);
+			return UE::OpenTrackIO::Cbor::Private::SetPropertyValue(StrProperty, Outer, Data, ArrayIndex, StringValue);
 		}
 
 		if (FNameProperty* NameProperty = CastField<FNameProperty>(Property))
 		{
-			return UE::OpenTrackIOCbor::Private::SetPropertyValue(NameProperty, Outer, Data, ArrayIndex, FName(*StringValue));
+			return UE::OpenTrackIO::Cbor::Private::SetPropertyValue(NameProperty, Outer, Data, ArrayIndex, FName(*StringValue));
 		}
 
 		if (FTextProperty* TextProperty = CastField<FTextProperty>(Property))
@@ -364,7 +414,7 @@ bool FOpenTrackIOCborStructDeserializerBackend::ReadProperty(FProperty* Property
 			{
 				TextValue = FText::FromString(StringValue);
 			}
-			return UE::OpenTrackIOCbor::Private::SetPropertyValue(TextProperty, Outer, Data, ArrayIndex, TextValue);
+			return UE::OpenTrackIO::Cbor::Private::SetPropertyValue(TextProperty, Outer, Data, ArrayIndex, TextValue);
 		}
 
 		if (FByteProperty* ByteProperty = CastField<FByteProperty>(Property))
@@ -380,7 +430,7 @@ bool FOpenTrackIOCborStructDeserializerBackend::ReadProperty(FProperty* Property
 				return false;
 			}
 
-			return UE::OpenTrackIOCbor::Private::SetPropertyValue(ByteProperty, Outer, Data, ArrayIndex, (uint8)Value);
+			return UE::OpenTrackIO::Cbor::Private::SetPropertyValue(ByteProperty, Outer, Data, ArrayIndex, (uint8)Value);
 		}
 
 		if (FEnumProperty* EnumProperty = CastField<FEnumProperty>(Property))
@@ -391,7 +441,7 @@ bool FOpenTrackIOCborStructDeserializerBackend::ReadProperty(FProperty* Property
 				return false;
 			}
 
-			if (void* ElementPtr = UE::OpenTrackIOCbor::Private::GetPropertyValuePtr(EnumProperty, Outer, Data, ArrayIndex))
+			if (void* ElementPtr = UE::OpenTrackIO::Cbor::Private::GetPropertyValuePtr(EnumProperty, Outer, Data, ArrayIndex))
 			{
 				EnumProperty->GetUnderlyingProperty()->SetIntPropertyValue(ElementPtr, Value);
 				return true;
@@ -402,27 +452,27 @@ bool FOpenTrackIOCborStructDeserializerBackend::ReadProperty(FProperty* Property
 
 		if (FClassProperty* ClassProperty = CastField<FClassProperty>(Property))
 		{
-			return UE::OpenTrackIOCbor::Private::SetPropertyValue(ClassProperty, Outer, Data, ArrayIndex, LoadObject<UClass>(NULL, *StringValue, NULL, LOAD_NoWarn));
+			return UE::OpenTrackIO::Cbor::Private::SetPropertyValue(ClassProperty, Outer, Data, ArrayIndex, LoadObject<UClass>(NULL, *StringValue, NULL, LOAD_NoWarn));
 		}
 
 		if (FSoftClassProperty* SoftClassProperty = CastField<FSoftClassProperty>(Property))
 		{
-			return UE::OpenTrackIOCbor::Private::SetPropertyValue(SoftClassProperty, Outer, Data, ArrayIndex, FSoftObjectPtr(FSoftObjectPath(StringValue)));
+			return UE::OpenTrackIO::Cbor::Private::SetPropertyValue(SoftClassProperty, Outer, Data, ArrayIndex, FSoftObjectPtr(FSoftObjectPath(StringValue)));
 		}
 
 		if (FObjectProperty* ObjectProperty = CastField<FObjectProperty>(Property))
 		{
-			return UE::OpenTrackIOCbor::Private::SetPropertyValue(ObjectProperty, Outer, Data, ArrayIndex, StaticFindObject(ObjectProperty->PropertyClass, nullptr, *StringValue));
+			return UE::OpenTrackIO::Cbor::Private::SetPropertyValue(ObjectProperty, Outer, Data, ArrayIndex, StaticFindObject(ObjectProperty->PropertyClass, nullptr, *StringValue));
 		}
 
 		if (FWeakObjectProperty* WeakObjectProperty = CastField<FWeakObjectProperty>(Property))
 		{
-			return UE::OpenTrackIOCbor::Private::SetPropertyValue(WeakObjectProperty, Outer, Data, ArrayIndex, FWeakObjectPtr(StaticFindObject(WeakObjectProperty->PropertyClass, nullptr, *StringValue)));
+			return UE::OpenTrackIO::Cbor::Private::SetPropertyValue(WeakObjectProperty, Outer, Data, ArrayIndex, FWeakObjectPtr(StaticFindObject(WeakObjectProperty->PropertyClass, nullptr, *StringValue)));
 		}
 
 		if (FSoftObjectProperty* SoftObjectProperty = CastField<FSoftObjectProperty>(Property))
 		{
-			return UE::OpenTrackIOCbor::Private::SetPropertyValue(SoftObjectProperty, Outer, Data, ArrayIndex, FSoftObjectPtr(FSoftObjectPath(StringValue)));
+			return UE::OpenTrackIO::Cbor::Private::SetPropertyValue(SoftObjectProperty, Outer, Data, ArrayIndex, FSoftObjectPtr(FSoftObjectPath(StringValue)));
 		}
 
 		UE_LOG(LogLiveLinkOpenTrackIO, Verbose, TEXT("String field %s with value '%s' is not supported in FProperty type %s (%s)"), *Property->GetFName().ToString(), *StringValue, *Property->GetClass()->GetName(), *GetDebugString());
@@ -442,11 +492,11 @@ bool FOpenTrackIOCborStructDeserializerBackend::ReadProperty(FProperty* Property
 
 		if (FByteProperty* ByteProperty = CastField<FByteProperty>(Property))
 		{
-			return UE::OpenTrackIOCbor::Private::SetPropertyValue(ByteProperty, Outer, Data, ArrayIndex, ByteValue);
+			return UE::OpenTrackIO::Cbor::Private::SetPropertyValue(ByteProperty, Outer, Data, ArrayIndex, ByteValue);
 		}
 		else if (FInt8Property* Int8Property = CastField<FInt8Property>(Property))
 		{
-			return UE::OpenTrackIOCbor::Private::SetPropertyValue(Int8Property, Outer, Data, ArrayIndex, (int8)ByteValue);
+			return UE::OpenTrackIO::Cbor::Private::SetPropertyValue(Int8Property, Outer, Data, ArrayIndex, (int8)ByteValue);
 		}
 
 		UE_LOG(LogLiveLinkOpenTrackIO, Verbose, TEXT("Error while deserializing field %s. Unexpected UProperty type %s. Expected a UByteProperty/UInt8Property to deserialize a TArray<uint8>/TArray<int8>"), *Property->GetFName().ToString(), *Property->GetClass()->GetName());
@@ -467,7 +517,7 @@ bool FOpenTrackIOCborStructDeserializerBackend::ReadProperty(FProperty* Property
 
 			if (FBoolProperty* BoolProperty = CastField<FBoolProperty>(Property))
 			{
-				return UE::OpenTrackIOCbor::Private::SetPropertyValue(BoolProperty, Outer, Data, ArrayIndex, LastContext.AsBool());
+				return UE::OpenTrackIO::Cbor::Private::SetPropertyValue(BoolProperty, Outer, Data, ArrayIndex, LastContext.AsBool());
 			}
 			UE_LOG(LogLiveLinkOpenTrackIO, Verbose, TEXT("Boolean field %s with value '%s' is not supported in FProperty type %s (%s)"), *Property->GetFName().ToString(), LastContext.AsBool() ? *(CoreTexts.True.ToString()) : *(CoreTexts.False.ToString()), *Property->GetClass()->GetName(), *GetDebugString());
 			return false;
@@ -476,13 +526,13 @@ bool FOpenTrackIOCborStructDeserializerBackend::ReadProperty(FProperty* Property
 		// Null
 		case ECborCode::Null:
 		{
-			return UE::OpenTrackIOCbor::Private::ClearPropertyValue(Property, Outer, Data, ArrayIndex);
+			return UE::OpenTrackIO::Cbor::Private::ClearPropertyValue(Property, Outer, Data, ArrayIndex);
 		}
 
 		// Float
 		case ECborCode::Value_4Bytes:
 		{
-			const bool bWasSet = UE::OpenTrackIOCbor::Private::SetNumericPropertyValue(Property, Outer, Data, ArrayIndex, LastContext.AsFloat());
+			const bool bWasSet = UE::OpenTrackIO::Cbor::Private::SetNumericPropertyValue(Property, Outer, Data, ArrayIndex, LastContext.AsFloat());
 
 			if (!bWasSet)
 			{
@@ -496,7 +546,7 @@ bool FOpenTrackIOCborStructDeserializerBackend::ReadProperty(FProperty* Property
 		// Double
 		case ECborCode::Value_8Bytes:
 		{
-			const bool bWasSet = UE::OpenTrackIOCbor::Private::SetNumericPropertyValue(Property, Outer, Data, ArrayIndex, LastContext.AsDouble());
+			const bool bWasSet = UE::OpenTrackIO::Cbor::Private::SetNumericPropertyValue(Property, Outer, Data, ArrayIndex, LastContext.AsDouble());
 
 			if (!bWasSet)
 			{
