Add a CVAR that allows users to control the modulo value used for fletcher-16 checksum.

The OpenTrackIO spec currently specifies a module of 255 but there is ongoing discussion about changing this to module 256.  The decision might not happen until after 5.6 locks so we are adding a CVar to allow us to adjust this post-release.

#rnx

diff --git a/Engine/Plugins/VirtualProduction/LiveLinkOpenTrackIO/Source/LiveLinkOpenTrackIO/Private/LiveLinkOpenTrackIOParser.cpp b/Engine/Plugins/VirtualProduction/LiveLinkOpenTrackIO/Source/LiveLinkOpenTrackIO/Private/LiveLinkOpenTrackIOParser.cpp
--- a/Engine/Plugins/VirtualProduction/LiveLinkOpenTrackIO/Source/LiveLinkOpenTrackIO/Private/LiveLinkOpenTrackIOParser.cpp
+++ b/Engine/Plugins/VirtualProduction/LiveLinkOpenTrackIO/Source/LiveLinkOpenTrackIO/Private/LiveLinkOpenTrackIOParser.cpp
@@ -28,6 +28,10 @@
 static float PayloadSizeLimitInMb = 64;
 static FAutoConsoleVariableRef  CVarPayloadSizeLimitInMb(TEXT("OpenTrackIO.PayloadSizeLimit"), PayloadSizeLimitInMb, TEXT("Limits sizes of payloads that we will accept. The value is in megabytes."));
 
+// Payload size limiter to make sure senders don't endlessly send us data without a last segment. 
+static int32 ChecksumModulo = 255;
+static FAutoConsoleVariableRef CVarChecksumModulo(TEXT("OpenTrackIO.ChecksumModulo"), ChecksumModulo, TEXT("The modulo value to use in the fletcher-16 checksum value."));
+
 void FLiveLinkOpenTrackIOPayload::AddBytes(uint32 InOffset, TArrayView<const uint8> InBytes)
 {
 	const float PayloadSizeLimitInFloatBytes = PayloadSizeLimitInMb * 1024 * 1024;
@@ -369,8 +373,8 @@ static uint16_t CalculateChecksum(TArrayView<const uint8> HeaderBytes, TArrayVie
 					C1 += C0;
 				} while (--BlockLen);
 
-				C0 %= 255;
-				C1 %= 255;
+				C0 %= ChecksumModulo;
+				C1 %= ChecksumModulo;
 			}
 		};
 
