The header sequence number of an OpenTrackIO packet will only increase after the full payload has been received.  We will check  the header sequence number with the previous sequence number to ensure they match when we have segmented payloads.

#rnx

diff --git a/Engine/Plugins/VirtualProduction/LiveLinkOpenTrackIO/Source/LiveLinkOpenTrackIO/Private/LiveLinkOpenTrackIOParser.cpp b/Engine/Plugins/VirtualProduction/LiveLinkOpenTrackIO/Source/LiveLinkOpenTrackIO/Private/LiveLinkOpenTrackIOParser.cpp
--- a/Engine/Plugins/VirtualProduction/LiveLinkOpenTrackIO/Source/LiveLinkOpenTrackIO/Private/LiveLinkOpenTrackIOParser.cpp
+++ b/Engine/Plugins/VirtualProduction/LiveLinkOpenTrackIO/Source/LiveLinkOpenTrackIO/Private/LiveLinkOpenTrackIOParser.cpp
@@ -531,9 +531,11 @@ bool GetHeaderAndPayloadFromBytes(TArrayView<const uint8> Bytes, FOpenTrackIOHea
 		return false;		
 	}
 
+	FLiveLinkOpenTrackIODatagramHeader& Header = OutPayloadContainer.GetMutableHeader();
+	uint16 PrevSequenceNumber = Header.SequenceNumber;
+	
 	/** Now read in the header to prepare parsing the payload. */
 	FMemoryReaderView Reader(Bytes);
-	FLiveLinkOpenTrackIODatagramHeader& Header = OutPayloadContainer.GetMutableHeader();
 	Reader << Header;
 
 	/** Confirm that the payload size is within the number of bytes given to us. */
@@ -556,7 +558,16 @@ bool GetHeaderAndPayloadFromBytes(TArrayView<const uint8> Bytes, FOpenTrackIOHea
 		return false;
 	}
 
-	OutPayloadContainer.GetMutablePayload().AddBytes(Header.SegmentOffset, Payload);
+	// According to the OpenTrackIO team, he sequence number will only increase when a full payload has been received. So the sequence number
+	// should match the previous sequence number.
+	FLiveLinkOpenTrackIOPayload& PayloadStorage = OutPayloadContainer.GetMutablePayload();
+	if (PayloadStorage.HasAnyPayloadData() && Header.SequenceNumber != PrevSequenceNumber)
+	{
+		UE_LOG(LogLiveLinkOpenTrackIO, Warning, TEXT("Invalid sequence number for segmented payload. Expected %d but received %d."), PrevSequenceNumber, Header.SequenceNumber);
+		return false;
+	}
+	
+	PayloadStorage.AddBytes(Header.SegmentOffset, Payload);
 
 	return true;
 }
diff --git a/Engine/Plugins/VirtualProduction/LiveLinkOpenTrackIO/Source/LiveLinkOpenTrackIO/Private/LiveLinkOpenTrackIOParser.h b/Engine/Plugins/VirtualProduction/LiveLinkOpenTrackIO/Source/LiveLinkOpenTrackIO/Private/LiveLinkOpenTrackIOParser.h
--- a/Engine/Plugins/VirtualProduction/LiveLinkOpenTrackIO/Source/LiveLinkOpenTrackIO/Private/LiveLinkOpenTrackIOParser.h
+++ b/Engine/Plugins/VirtualProduction/LiveLinkOpenTrackIO/Source/LiveLinkOpenTrackIO/Private/LiveLinkOpenTrackIOParser.h
@@ -32,6 +32,15 @@ struct FLiveLinkOpenTrackIOPayload
 	 */
 	bool IsComplete() const;
 
+	/**
+	 * Returns true if any payload data exists on this struct. This is useful for detecting continuity of
+	 * data between header values. 
+	 */
+	bool HasAnyPayloadData() const
+	{
+		return Sections.Num() > 0;
+	}
+	
 	/**
 	 * Return a view into the data such that it can be passed into the CBOR or JSON parsers.
 	 *
