Per OpenTrackIO team, change the sequence number to be from 0-65535 and change checksum to using 256 for modulo function.

diff --git a/Engine/Plugins/VirtualProduction/LiveLinkOpenTrackIO/Source/LiveLinkOpenTrackIO/Private/LiveLinkOpenTrackIOParser.cpp b/Engine/Plugins/VirtualProduction/LiveLinkOpenTrackIO/Source/LiveLinkOpenTrackIO/Private/LiveLinkOpenTrackIOParser.cpp
--- a/Engine/Plugins/VirtualProduction/LiveLinkOpenTrackIO/Source/LiveLinkOpenTrackIO/Private/LiveLinkOpenTrackIOParser.cpp
+++ b/Engine/Plugins/VirtualProduction/LiveLinkOpenTrackIO/Source/LiveLinkOpenTrackIO/Private/LiveLinkOpenTrackIOParser.cpp
@@ -29,7 +29,7 @@ static float PayloadSizeLimitInMb = 64;
 static FAutoConsoleVariableRef  CVarPayloadSizeLimitInMb(TEXT("OpenTrackIO.PayloadSizeLimit"), PayloadSizeLimitInMb, TEXT("Limits sizes of payloads that we will accept. The value is in megabytes."));
 
 // Payload size limiter to make sure senders don't endlessly send us data without a last segment. 
-static int32 ChecksumModulo = 255;
+static int32 ChecksumModulo = 256;
 static FAutoConsoleVariableRef CVarChecksumModulo(TEXT("OpenTrackIO.ChecksumModulo"), ChecksumModulo, TEXT("The modulo value to use in the fletcher-16 checksum value."));
 
 void FLiveLinkOpenTrackIOPayload::AddBytes(uint32 InOffset, TArrayView<const uint8> InBytes)
diff --git a/Engine/Plugins/VirtualProduction/LiveLinkOpenTrackIO/Source/LiveLinkOpenTrackIO/Private/LiveLinkOpenTrackIOTranscoder.h b/Engine/Plugins/VirtualProduction/LiveLinkOpenTrackIO/Source/LiveLinkOpenTrackIO/Private/LiveLinkOpenTrackIOTranscoder.h
--- a/Engine/Plugins/VirtualProduction/LiveLinkOpenTrackIO/Source/LiveLinkOpenTrackIO/Private/LiveLinkOpenTrackIOTranscoder.h
+++ b/Engine/Plugins/VirtualProduction/LiveLinkOpenTrackIO/Source/LiveLinkOpenTrackIO/Private/LiveLinkOpenTrackIOTranscoder.h
@@ -43,7 +43,7 @@ struct FLiveLinkOpenTrackIOCache
 		ExpectedSequenceNumber = SequenceNumber + 1;
 		if (ExpectedSequenceNumber > UINT16_MAX)
 		{
-			ExpectedSequenceNumber = 1;
+			ExpectedSequenceNumber = 0;
 		}
 		LastDataReceiveTimeInSeconds = FPlatformTime::Seconds();
 	}
