OpenTrackIO: Update to upcoming OTrIO version with a changed multicast address.

The full multicast address is now exposed to the user, and defaults to the expected new value in the standard (239.135.1.X instead of 235.135.1.X).

diff --git a/Engine/Plugins/VirtualProduction/LiveLinkOpenTrackIO/Source/LiveLinkOpenTrackIO/Private/LiveLinkOpenTrackIOConnectionSettings.h b/Engine/Plugins/VirtualProduction/LiveLinkOpenTrackIO/Source/LiveLinkOpenTrackIO/Private/LiveLinkOpenTrackIOConnectionSettings.h
--- a/Engine/Plugins/VirtualProduction/LiveLinkOpenTrackIO/Source/LiveLinkOpenTrackIO/Private/LiveLinkOpenTrackIOConnectionSettings.h
+++ b/Engine/Plugins/VirtualProduction/LiveLinkOpenTrackIO/Source/LiveLinkOpenTrackIO/Private/LiveLinkOpenTrackIOConnectionSettings.h
@@ -1,6 +1,7 @@
 // Copyright Epic Games, Inc. All Rights Reserved.
 
 #pragma once
+
 #include "LiveLinkSourceSettings.h"
 
 #include "LiveLinkOpenTrackIOConnectionSettings.generated.h"
@@ -31,24 +32,31 @@ struct LIVELINKOPENTRACKIO_API FLiveLinkOpenTrackIOConnectionSettings
 	UPROPERTY(EditAnywhere, Category = "Settings")
 	ELiveLinkOpenTrackIONetworkProtocol Protocol = ELiveLinkOpenTrackIONetworkProtocol::Multicast;
 
-	/**
-	 * The Source Number will be used as the last octet of the multicast address 235.135.1.[SourceNumber], per the opentrackio.org spec.
-	 * Use values in the range [1,200].
-	 */
+	/** Multicast Group. The last octet is the SourceNumber and should stay within the range [1,200], per the OpenTrackIO spec. */
 	UPROPERTY(EditAnywhere, Category = "Settings", meta = (
 		EditCondition = "Protocol == ELiveLinkOpenTrackIONetworkProtocol::Multicast",
-		EditConditionHides,
-		ClampMin = "1", ClampMax = "200",
-		UIMin = "1", UIMax = "200"
+		EditConditionHides
 		))
-	uint8 SourceNumber = 1;
+	FString MulticastGroup = TEXT("239.135.1.1");
 
-	/** Unicast port number. 0 will self assign, so it is suggested to pick a port that can be used in your system. */
+	/** Multicast Port. */
 	UPROPERTY(EditAnywhere, Category = "Settings", meta = (
-		EditCondition = "Protocol == ELiveLinkOpenTrackIONetworkProtocol::Unicast",
+		EditCondition = "Protocol == ELiveLinkOpenTrackIONetworkProtocol::Multicast",
 		EditConditionHides,
 		ClampMin = "1", ClampMax = "65535",
 		UIMin = "1", UIMax = "65535"
 		))
+	uint16 MulticastPort = 55555;
+
+	/** Local interface (used in both Unicast and Multicast protocols). */
+	UPROPERTY(EditAnywhere, Category = "Settings")
+	FString LocalInterface = TEXT("0.0.0.0");
+
+	/** Unicast Port. 0 means self-assign. Use a port number available in your system. */
+	UPROPERTY(EditAnywhere, Category = "Settings", meta = (
+		EditCondition = "Protocol == ELiveLinkOpenTrackIONetworkProtocol::Unicast",
+		EditConditionHides
+		))
 	uint16 UnicastPort = 0;
+
 };
diff --git a/Engine/Plugins/VirtualProduction/LiveLinkOpenTrackIO/Source/LiveLinkOpenTrackIO/Private/LiveLinkOpenTrackIOSource.cpp b/Engine/Plugins/VirtualProduction/LiveLinkOpenTrackIO/Source/LiveLinkOpenTrackIO/Private/LiveLinkOpenTrackIOSource.cpp
--- a/Engine/Plugins/VirtualProduction/LiveLinkOpenTrackIO/Source/LiveLinkOpenTrackIO/Private/LiveLinkOpenTrackIOSource.cpp
+++ b/Engine/Plugins/VirtualProduction/LiveLinkOpenTrackIO/Source/LiveLinkOpenTrackIO/Private/LiveLinkOpenTrackIOSource.cpp
@@ -106,20 +106,6 @@ void JoinedToGroup(const FIPv4Endpoint& UnicastEndpoint, const FIPv4Endpoint& Mu
 #endif
 }
 
-/**
- * Makes the Multicast Endpoint to use per the opentrackio spec.
- * Source Number will be clamped to 1..200 per opentrackio.org.
- * 
- * @param InSourceNumber  OpenTrackIO Source Number to encode into the fourth octet.
- * @param InPortNumber    UDP multicast port number.
- *
- * @return FIPv4Endpoint  The multicast address (235.135.1.<SourceNumber>:<PortNumber>).
- */
-FIPv4Endpoint MakeMulticastAddressFromSourceNumber(uint8 InSourceNumber, uint16 InPortNumber = 55555)
-{
-	const uint8 SourceNumber = FMath::Clamp(InSourceNumber, 1, 200);
-	return FIPv4Endpoint(FIPv4Address(235, 135, 1, SourceNumber), InPortNumber);
-}
 }
 
 FLiveLinkOpenTrackIOSource::FLiveLinkOpenTrackIOSource(FLiveLinkOpenTrackIOConnectionSettings InConnectionSettings)
@@ -272,28 +258,59 @@ void FLiveLinkOpenTrackIOSource::Update()
 bool FLiveLinkOpenTrackIOSource::ParseEndpointFromSourceSettings()
 {
 	check(IsValid(SavedSourceSettings));
-	
+
 	using namespace UE::LiveLinkOpenTrackIO::Private;
 
-	TOptional<FIPv4Endpoint> InUniEndpoint = ParseEndpoint(SavedSourceSettings->UnicastEndpoint);
+	// Unicast settings
+
+	// Note: The Unicast port number is not used when using multicast, only the address (which is used as an interface address to bind to).
+	UnicastEndpoint.Port = ConnectionSettings.UnicastPort;
+
+	TOptional<FIPv4Endpoint> UniEndpoint = ParseEndpoint(SavedSourceSettings->UnicastEndpoint);
 
-	if (!InUniEndpoint.IsSet())
+	if (UniEndpoint.IsSet())
 	{
-		return false;
+		UnicastEndpoint = *UniEndpoint;
 	}
 
-	UnicastEndpoint = *InUniEndpoint;
+	// Multicast Settings
 
-	// Note: The Unicast port number is not used when using multicast, only the address (which is used as an interface address to bind to).
-	UnicastEndpoint.Port = ConnectionSettings.UnicastPort;
+	TOptional<FIPv4Endpoint> MultiEndpoint = ParseEndpoint(SavedSourceSettings->MulticastEndpoint);
+
+	if (MultiEndpoint.IsSet())
+	{
+		MulticastEndpoint = *MultiEndpoint;
+	}
+
+	switch (SavedSourceSettings->Protocol)
+	{
+	case ELiveLinkOpenTrackIONetworkProtocol::Unicast:
+
+		if (!UniEndpoint.IsSet())
+		{
+			return false;
+		}
+		break;
+
+	case ELiveLinkOpenTrackIONetworkProtocol::Multicast:
+
+		if (!UniEndpoint.IsSet() || !MultiEndpoint.IsSet() || !MultiEndpoint->Address.IsMulticastAddress())
+		{
+			return false;
+		}
+		break;
 
-	MulticastEndpoint = MakeMulticastAddressFromSourceNumber(ConnectionSettings.SourceNumber, SavedSourceSettings->MulticastPort);
+	default:
+		return false;
+	}
 
 	return true;
 }
 
 void FLiveLinkOpenTrackIOSource::InitializeSettings(ULiveLinkSourceSettings* Settings)
 {
+	using namespace UE::LiveLinkOpenTrackIO::Private;
+
 	ILiveLinkSource::InitializeSettings(Settings);
 
 	if (ULiveLinkOpenTrackIOSourceSettings* OpenTrackIOSettings = Cast<ULiveLinkOpenTrackIOSourceSettings>(Settings))
@@ -301,6 +318,26 @@ void FLiveLinkOpenTrackIOSource::InitializeSettings(ULiveLinkSourceSettings* Set
 		// Cache this for details to properties with the right context.
 		OpenTrackIOSettings->Protocol = ConnectionSettings.Protocol;
 
+		// Multicast
+		{
+			FIPv4Endpoint Endpoint;
+
+			FIPv4Address::Parse(ConnectionSettings.MulticastGroup, Endpoint.Address);
+			Endpoint.Port = ConnectionSettings.MulticastPort;
+
+			OpenTrackIOSettings->MulticastEndpoint = Endpoint.ToString();
+		}
+
+		// Unicast
+		{
+			FIPv4Endpoint Endpoint;
+
+			FIPv4Address::Parse(ConnectionSettings.LocalInterface, Endpoint.Address);
+			Endpoint.Port = ConnectionSettings.UnicastPort;
+
+			OpenTrackIOSettings->UnicastEndpoint = Endpoint.ToString();
+		}
+
 		SavedSourceSettings = OpenTrackIOSettings;
 
 		bResetRequested = true;
@@ -311,29 +348,6 @@ void FLiveLinkOpenTrackIOSource::InitializeSettings(ULiveLinkSourceSettings* Set
 	}
 }
 
-void FLiveLinkOpenTrackIOSource::OnSettingsChanged(ULiveLinkSourceSettings* Settings, const FPropertyChangedEvent& PropertyChangedEvent)
-{
-	ILiveLinkSource::OnSettingsChanged(Settings, PropertyChangedEvent);
-
-	const FProperty* const MemberProperty = PropertyChangedEvent.MemberProperty;
-	const FProperty* const Property = PropertyChangedEvent.Property;
-	if (Property && MemberProperty && (PropertyChangedEvent.ChangeType != EPropertyChangeType::Interactive))
-	{
-		if (ensureMsgf(IsValid(SavedSourceSettings), TEXT("The LiveLinkOpenTrackIO source settings have become invalid.  Please re-create the source.")))
-		{
-			// If the user has changed the endpoint addresses then we need to reset the socket connection and re-establish receivers. 
-			const FName PropertyName = Property->GetFName();
-			const FName MemberPropertyName = MemberProperty->GetFName();
-			if (PropertyName == GET_MEMBER_NAME_CHECKED(ULiveLinkOpenTrackIOSourceSettings, UnicastEndpoint)
-				|| PropertyName == GET_MEMBER_NAME_CHECKED(ULiveLinkOpenTrackIOSourceSettings, MulticastPort))
-			{
-				// Start the reset process.
-				bResetRequested = true;
-			}
-		}
-	}
-}
-	
 void FLiveLinkOpenTrackIOSource::HandleInboundData(const TSharedPtr<FArrayReader, ESPMode::ThreadSafe>& InData, const FIPv4Endpoint& InSender)
 {
 	TRACE_CPUPROFILER_EVENT_SCOPE(OpenTrackIO::HandleInboundData);
@@ -639,7 +653,7 @@ bool FLiveLinkOpenTrackIOSource::OpenUnicastSocket()
 
 bool FLiveLinkOpenTrackIOSource::OpenSockets()
 {
-	switch (ConnectionSettings.Protocol)
+	switch (SavedSourceSettings->Protocol)
 	{
 	case ELiveLinkOpenTrackIONetworkProtocol::Unicast:
 	{
diff --git a/Engine/Plugins/VirtualProduction/LiveLinkOpenTrackIO/Source/LiveLinkOpenTrackIO/Private/LiveLinkOpenTrackIOSource.h b/Engine/Plugins/VirtualProduction/LiveLinkOpenTrackIO/Source/LiveLinkOpenTrackIO/Private/LiveLinkOpenTrackIOSource.h
--- a/Engine/Plugins/VirtualProduction/LiveLinkOpenTrackIO/Source/LiveLinkOpenTrackIO/Private/LiveLinkOpenTrackIOSource.h
+++ b/Engine/Plugins/VirtualProduction/LiveLinkOpenTrackIO/Source/LiveLinkOpenTrackIO/Private/LiveLinkOpenTrackIOSource.h
@@ -60,7 +60,6 @@ public:
 	// Begin ILiveLinkSource Interface
 	virtual void ReceiveClient(ILiveLinkClient* InClient, FGuid InSourceGuid) override;
 	virtual void InitializeSettings(ULiveLinkSourceSettings* Settings) override;
-	virtual void OnSettingsChanged(ULiveLinkSourceSettings* Settings, const FPropertyChangedEvent& PropertyChangedEvent) override;
 
 	virtual bool IsSourceStillValid() const override;
 
diff --git a/Engine/Plugins/VirtualProduction/LiveLinkOpenTrackIO/Source/LiveLinkOpenTrackIO/Public/LiveLinkOpenTrackIOSourceSettings.h b/Engine/Plugins/VirtualProduction/LiveLinkOpenTrackIO/Source/LiveLinkOpenTrackIO/Public/LiveLinkOpenTrackIOSourceSettings.h
--- a/Engine/Plugins/VirtualProduction/LiveLinkOpenTrackIO/Source/LiveLinkOpenTrackIO/Public/LiveLinkOpenTrackIOSourceSettings.h
+++ b/Engine/Plugins/VirtualProduction/LiveLinkOpenTrackIO/Source/LiveLinkOpenTrackIO/Public/LiveLinkOpenTrackIOSourceSettings.h
@@ -24,19 +24,19 @@ class LIVELINKOPENTRACKIO_API ULiveLinkOpenTrackIOSourceSettings : public ULiveL
 public:
 
 	/** Protocol being used for the connection. */
-	UPROPERTY()
+	UPROPERTY(VisibleAnywhere, Category = Transport)
 	ELiveLinkOpenTrackIONetworkProtocol Protocol = ELiveLinkOpenTrackIONetworkProtocol::Multicast;
 
 	/**
-	 * The Multicast port to receive packets on.
+	 * The Multicast endpoint to receive packets on.
 	 *
-	 * The default value 55555 per OpenTrackIO specifications.
+	 * The default value is 239.135.1.1:55555 per OpenTrackIO specifications.
 	 */
-	UPROPERTY(EditAnywhere, Category = Transport, meta = (
+	UPROPERTY(VisibleAnywhere, Category = Transport, meta = (
 		EditCondition = "Protocol == ELiveLinkOpenTrackIONetworkProtocol::Multicast",
 		EditConditionHides
 	))
-	uint16 MulticastPort = 55555;
+	FString MulticastEndpoint = TEXT("239.135.1.1:55555");
 
 	/**
 	 * The IP endpoint to listen to and send packets from.  
@@ -46,7 +46,7 @@ public:
 	 * and all available network adapters on other operating systems.
 	 *
 	 */
-	UPROPERTY(EditAnywhere, Category=Transport)
+	UPROPERTY(VisibleAnywhere, Category=Transport)
 	FString UnicastEndpoint = TEXT("0.0.0.0:0");
 
 	/**
