Handle invalid CBOR data that could be received externally by the engine.  The original bug report is an assert in a condition that is not supported by the deserializer.  This properly handles this edge case and reports failure back up the caller which should ignore the data and continue.

#Synced-CL 49278929

diff --git a/Engine/Plugins/VirtualProduction/LiveLinkOpenTrackIO/Source/LiveLinkOpenTrackIO/Private/OpenTrackIOCborStructDeserializerBackend.cpp b/Engine/Plugins/VirtualProduction/LiveLinkOpenTrackIO/Source/LiveLinkOpenTrackIO/Private/OpenTrackIOCborStructDeserializerBackend.cpp
--- a/Engine/Plugins/VirtualProduction/LiveLinkOpenTrackIO/Source/LiveLinkOpenTrackIO/Private/OpenTrackIOCborStructDeserializerBackend.cpp
+++ b/Engine/Plugins/VirtualProduction/LiveLinkOpenTrackIO/Source/LiveLinkOpenTrackIO/Private/OpenTrackIOCborStructDeserializerBackend.cpp
@@ -214,6 +214,36 @@ const FString& FOpenTrackIOCborStructDeserializerBackend::GetLastErrorMessage()
 	return Dummy;
 }
 
+namespace UE::OpenTrackIOCborStructDeserializerBackend::Private
+{
+
+FString ConvertCborTypeToString(ECborCode CborCode)
+{
+	switch (CborCode)
+	{
+		case ECborCode::Uint:
+			return TEXT("Uint");
+		case ECborCode::Int:
+			return TEXT("Int");
+		case ECborCode::ByteString:
+			return TEXT("ByteString");
+		case ECborCode::TextString:
+			return TEXT("TextString");						
+		case ECborCode::Array:
+			return TEXT("Array");
+		case ECborCode::Map:
+			return TEXT("Map");
+		case ECborCode::Tag:
+			return TEXT("Tag");
+		case ECborCode::Prim:
+			return TEXT("Prim");
+		// Fallthrough
+	};
+	return TEXT("InvalidType");
+}
+
+} // namespace UE::CborStructDeserializerBackend::Private
+
 bool FOpenTrackIOCborStructDeserializerBackend::GetNextToken(EStructDeserializerBackendTokens& OutToken)
 {
 	LastMapKey.Reset();
@@ -248,17 +278,29 @@ bool FOpenTrackIOCborStructDeserializerBackend::GetNextToken(EStructDeserializer
 		return true;
 	}
 
-	// if after reading the last context, the parent context is a map with an odd length, we just read a key
+	// We've read an odd number of items from this map so far. So the next value is the key and it must be a string.
 	if (CborReader.GetContext().MajorType() == ECborCode::Map && (CborReader.GetContext().AsLength() & 1))
 	{
-		// Should be a string
-		check(LastContext.MajorType() == ECborCode::TextString);
-		LastMapKey = LastContext.AsString();
+		// This code is designed to deserialize CBOR data into Unreal Engine structs. If this is not a string then we are somehow not dealing with UE
+		// structs and this would be an error. Previously this code would just assert and crash.  This gracefully handles this case and passes the
+		// failure down to the caller. 
+		if (LastContext.MajorType() == ECborCode::TextString)
+		{
+			LastMapKey = LastContext.AsString();
 
-		// Read next and carry on
-		if (!CborReader.ReadNext(LastContext))
+			// Read next and carry on
+			if (!CborReader.ReadNext(LastContext))
+			{
+				OutToken = LastContext.IsError() ? EStructDeserializerBackendTokens::Error : EStructDeserializerBackendTokens::None;
+				return false;
+			}
+		}
+		else
 		{
-			OutToken = LastContext.IsError() ? EStructDeserializerBackendTokens::Error : EStructDeserializerBackendTokens::None;
+			OutToken = EStructDeserializerBackendTokens::Error;
+			FString CborMajorTypeAsString = UE::OpenTrackIOCborStructDeserializerBackend::Private::ConvertCborTypeToString(LastContext.MajorType());
+			uint8 CborMajorTypeAsInt = static_cast<uint8>(LastContext.MajorType());
+			UE_LOG(LogSerialization, Error, TEXT("Encountered invalided context while parsing CBOR map. Expected a string type but got code %d (%s)"), CborMajorTypeAsInt, *CborMajorTypeAsString);
 			return false;
 		}
 	}
diff --git a/Engine/Source/Runtime/Serialization/Private/Backends/CborStructDeserializerBackend.cpp b/Engine/Source/Runtime/Serialization/Private/Backends/CborStructDeserializerBackend.cpp
--- a/Engine/Source/Runtime/Serialization/Private/Backends/CborStructDeserializerBackend.cpp
+++ b/Engine/Source/Runtime/Serialization/Private/Backends/CborStructDeserializerBackend.cpp
@@ -37,6 +37,36 @@ const FString& FCborStructDeserializerBackend::GetLastErrorMessage() const
 	return Dummy;
 }
 
+namespace UE::CborStructDeserializerBackend::Private
+{
+
+FString ConvertCborTypeToString(ECborCode CborCode)
+{
+	switch (CborCode)
+	{
+		case ECborCode::Uint:
+			return TEXT("Uint");
+		case ECborCode::Int:
+			return TEXT("Int");
+		case ECborCode::ByteString:
+			return TEXT("ByteString");
+		case ECborCode::TextString:
+			return TEXT("TextString");						
+		case ECborCode::Array:
+			return TEXT("Array");
+		case ECborCode::Map:
+			return TEXT("Map");
+		case ECborCode::Tag:
+			return TEXT("Tag");
+		case ECborCode::Prim:
+			return TEXT("Prim");
+		// Fallthrough
+	};
+	return TEXT("InvalidType");
+}
+
+} // namespace UE::CborStructDeserializerBackend::Private
+
 bool FCborStructDeserializerBackend::GetNextToken(EStructDeserializerBackendTokens& OutToken)
 {
 	LastMapKey.Reset();
@@ -71,17 +101,29 @@ bool FCborStructDeserializerBackend::GetNextToken(EStructDeserializerBackendToke
 		return true;
 	}
 
-	// if after reading the last context, the parent context is a map with an odd length, we just read a key
+	// We've read an odd number of items from this map so far. So the next value is the key and it must be a string.
 	if (CborReader.GetContext().MajorType() == ECborCode::Map && (CborReader.GetContext().AsLength() & 1))
 	{
-		// Should be a string
-		check(LastContext.MajorType() == ECborCode::TextString);
-		LastMapKey = LastContext.AsString();
+		// This code is designed to deserialize CBOR data into Unreal Engine structs. If this is not a string then we are somehow not dealing with UE
+		// structs and this would be an error. Previously this code would just assert and crash.  This gracefully handles this case and passes the
+		// failure down to the caller. 		
+		if (LastContext.MajorType() == ECborCode::TextString)
+		{
+			LastMapKey = LastContext.AsString();
 
-		// Read next and carry on
-		if (!CborReader.ReadNext(LastContext))
+			// Read next and carry on
+			if (!CborReader.ReadNext(LastContext))
+			{
+				OutToken = LastContext.IsError() ? EStructDeserializerBackendTokens::Error : EStructDeserializerBackendTokens::None;
+				return false;
+			}
+		}
+		else
 		{
-			OutToken = LastContext.IsError() ? EStructDeserializerBackendTokens::Error : EStructDeserializerBackendTokens::None;
+			OutToken = EStructDeserializerBackendTokens::Error;
+			FString CborMajorTypeAsString = UE::CborStructDeserializerBackend::Private::ConvertCborTypeToString(LastContext.MajorType());
+			uint8 CborMajorTypeAsInt = static_cast<uint8>(LastContext.MajorType());
+			UE_LOG(LogSerialization, Error, TEXT("Encountered invalided context while parsing CBOR map. Expected a string type but got code %d (%s)"), CborMajorTypeAsInt, *CborMajorTypeAsString);
 			return false;
 		}
 	}
