OpenTrackIO: Fix hang when shutting down live link source.

#rnx

#Synced-CL 51080076

diff --git a/Engine/Plugins/VirtualProduction/LiveLinkOpenTrackIO/Source/LiveLinkOpenTrackIO/Private/LiveLinkOpenTrackIOSource.cpp b/Engine/Plugins/VirtualProduction/LiveLinkOpenTrackIO/Source/LiveLinkOpenTrackIO/Private/LiveLinkOpenTrackIOSource.cpp
--- a/Engine/Plugins/VirtualProduction/LiveLinkOpenTrackIO/Source/LiveLinkOpenTrackIO/Private/LiveLinkOpenTrackIOSource.cpp
+++ b/Engine/Plugins/VirtualProduction/LiveLinkOpenTrackIO/Source/LiveLinkOpenTrackIO/Private/LiveLinkOpenTrackIOSource.cpp
@@ -543,11 +543,23 @@ void FLiveLinkOpenTrackIOSource::CloseSockets()
 {
 	StopUdpReceivers();
 
+	// Close sockets to unblock receiver threads.
+
+	if (MulticastSocket)
+	{
+		MulticastSocket->Close();
+	}
+
+	if (UnicastSocket)
+	{
+		UnicastSocket->Close();
+	}
+
+	// Join receiver threads. Safe because sockets are closed.
 	MulticastReceiver.Reset();
 	UnicastReceiver.Reset();
 
-	// Destroy the sockets
-
+	// Destroy socket objects
 	ISocketSubsystem* SocketSubsystem = ISocketSubsystem::Get(PLATFORM_SOCKETSUBSYSTEM);
 	check(SocketSubsystem);
 
@@ -562,8 +574,7 @@ void FLiveLinkOpenTrackIOSource::CloseSockets()
 		UnicastSocket = nullptr;
 	}
 
-	// It should be safe to clear this because the Udp threads are guaranteed to be shut down at this point,
-	// since the call to .Reset() earlier in this fuction will block until they are destroyed.
+	// Safe to clear because UDP threads are joined above.
 	OpenTrackIOCacheMap.Reset();
 }
 
