Make CameraData a non-optional type as that allows us to support the minimal OpenTrackIO spec.  We did require that users abide by the recommended schema.

#Synced-CL 51338485

diff --git a/Engine/Plugins/VirtualProduction/LiveLinkOpenTrackIO/Source/LiveLinkOpenTrackIO/Private/LiveLinkOpenTrackIOTranscoder.cpp b/Engine/Plugins/VirtualProduction/LiveLinkOpenTrackIO/Source/LiveLinkOpenTrackIO/Private/LiveLinkOpenTrackIOTranscoder.cpp
--- a/Engine/Plugins/VirtualProduction/LiveLinkOpenTrackIO/Source/LiveLinkOpenTrackIO/Private/LiveLinkOpenTrackIOTranscoder.cpp
+++ b/Engine/Plugins/VirtualProduction/LiveLinkOpenTrackIO/Source/LiveLinkOpenTrackIO/Private/LiveLinkOpenTrackIOTranscoder.cpp
@@ -1,20 +1,19 @@
 // Copyright Epic Games, Inc. All Rights Reserved.
-// 
+//
 #include "LiveLinkOpenTrackIOTranscoder.h"
 
 #include "Containers/UnrealString.h"
-#include "LiveLinkOpenTrackIOLiveLinkTypes.h"
 #include "LiveLinkOpenTrackIOConnectionSettings.h"
 #include "LiveLinkOpenTrackIOConversions.h"
+#include "LiveLinkOpenTrackIOLiveLinkTypes.h"
 
 #include "Misc/App.h"
 #include "Misc/QualifiedFrameTime.h"
 #include "Roles/LiveLinkCameraTypes.h"
 
-
 bool FLiveLinkOpenTrackIOCache::IsPacketInSequence(uint16 SequenceNumber, const FFrameRate& InRate) const
 {
-	// Initially the expected sequence number is -1, which indicates we have not received our first packet yet. 
+	// Initially the expected sequence number is -1, which indicates we have not received our first packet yet.
 	if (ExpectedSequenceNumber > 0 && SequenceNumber < ExpectedSequenceNumber)
 	{
 		// Ensure a minimum rate. Also avoids divide by zero if the external producer sends zero rate.
@@ -34,7 +33,7 @@ bool FLiveLinkOpenTrackIOCache::IsPacketInSequence(uint16 SequenceNumber, const
 }
 
 FName FLiveLinkOpenTrackIOCache::GetSubjectNameFromData(const FString& InSubjectName, const FLiveLinkOpenTrackIOData& InData)
-{	
+{
 	// If no InSubjectName hint is provided, fallback names will be automatically generated based on the data.
 
 	if (!InSubjectName.IsEmpty() && (InSubjectName != FLiveLinkOpenTrackIOConnectionSettings::AutoSubjectName))
@@ -43,10 +42,10 @@ FName FLiveLinkOpenTrackIOCache::GetSubjectNameFromData(const FString& InSubject
 	}
 
 	// First fallback: Static camera
-	if (StaticCamera.IsSet())
+	if (StaticCamera.IsValid())
 	{
 		const FString Postfix = FString::Printf(TEXT("%03lld"), InData.SourceNumber);
-		return UE::OpenTrackIO::ConvertTypeToFName(StaticCamera.GetValue(), Postfix);
+		return UE::OpenTrackIO::ConvertTypeToFName(StaticCamera, Postfix);
 	}
 
 	// Next fallback: Static lens
@@ -74,7 +73,6 @@ FName FLiveLinkOpenTrackIOCache::GetTransformName(const FLiveLinkOpenTrackIOTran
 	return FName(InTransform.Id);
 }
 
-
 FLiveLinkStaticDataStruct FLiveLinkOpenTrackIOCache::MakeStaticData(const FLiveLinkOpenTrackIOData& Data, bool bApplyXform)
 {
 	/** Lens class inherits from the camera data so we use that as our base for the static struct. */
@@ -111,14 +109,14 @@ FLiveLinkStaticDataStruct FLiveLinkOpenTrackIOCache::MakeStaticData(const FLiveL
 		// because some cameras with built-in undistortion may update this value dynamically.
 
 		if ((Data.Static.Camera.ActiveSensorPhysicalDimensions.Height.IsSet()
-			&& Data.Static.Camera.ActiveSensorPhysicalDimensions.Width.IsSet()))
+				&& Data.Static.Camera.ActiveSensorPhysicalDimensions.Width.IsSet()))
 		{
 			NewStaticData.bIsDynamicFilmbackSupported = true;
-		} 
-		else if (StaticCamera.IsSet())
+		}
+		else
 		{
-			if (StaticCamera->ActiveSensorPhysicalDimensions.Height.IsSet()
-				&& StaticCamera->ActiveSensorPhysicalDimensions.Width.IsSet())
+			if (StaticCamera.ActiveSensorPhysicalDimensions.Height.IsSet()
+				&& StaticCamera.ActiveSensorPhysicalDimensions.Width.IsSet())
 			{
 				NewStaticData.bIsDynamicFilmbackSupported = true;
 			}
@@ -147,12 +145,11 @@ FLiveLinkStaticDataStruct FLiveLinkOpenTrackIOCache::MakeStaticData(const FLiveL
 	return StaticDataStruct;
 }
 
-
 FLiveLinkFrameDataStruct FLiveLinkOpenTrackIOCache::MakeFrameData(const FLiveLinkOpenTrackIOData& Data, bool bApplyXform)
 {
 	using namespace LiveLinkOpenTrackIOConversions;
 
-	FTransform FinalTransform; 
+	FTransform FinalTransform;
 	if (bApplyXform)
 	{
 		/** Compute the final transform by concatentating the transforms in the frame data. */
@@ -168,8 +165,7 @@ FLiveLinkFrameDataStruct FLiveLinkOpenTrackIOCache::MakeFrameData(const FLiveLin
 	const FTransform PupilOffset(
 		FQuat::Identity,
 		ToUnrealTranslation(EntrancePupilAsXYZ),
-		FVector(1, 1, 1)
-	);
+		FVector(1, 1, 1));
 
 	FinalTransform = PupilOffset * FinalTransform;
 
@@ -187,9 +183,9 @@ FLiveLinkFrameDataStruct FLiveLinkOpenTrackIOCache::MakeFrameData(const FLiveLin
 		SampleTC = Data.Timing.Timecode.GetQualifiedFrameTime();
 	}
 
-	const FLiveLinkOpenTrackIOStaticCamera* Camera = StaticCamera.IsSet() ? &StaticCamera.GetValue() : nullptr;
+	const FLiveLinkOpenTrackIOStaticCamera* Camera = &StaticCamera;
 	const FLiveLinkOpenTrackIOStaticLens* Lens = StaticLens.IsSet() ? &StaticLens.GetValue() : nullptr;
-		
+
 	FLiveLinkFrameDataStruct FrameDataStruct = FLiveLinkFrameDataStruct(FLiveLinkOpenTrackIOFrameData::StaticStruct());
 	FLiveLinkOpenTrackIOFrameData& CameraAndLensData = *FrameDataStruct.Cast<FLiveLinkOpenTrackIOFrameData>();
 
diff --git a/Engine/Plugins/VirtualProduction/LiveLinkOpenTrackIO/Source/LiveLinkOpenTrackIO/Private/LiveLinkOpenTrackIOTranscoder.h b/Engine/Plugins/VirtualProduction/LiveLinkOpenTrackIO/Source/LiveLinkOpenTrackIO/Private/LiveLinkOpenTrackIOTranscoder.h
--- a/Engine/Plugins/VirtualProduction/LiveLinkOpenTrackIO/Source/LiveLinkOpenTrackIO/Private/LiveLinkOpenTrackIOTranscoder.h
+++ b/Engine/Plugins/VirtualProduction/LiveLinkOpenTrackIO/Source/LiveLinkOpenTrackIO/Private/LiveLinkOpenTrackIOTranscoder.h
@@ -16,19 +16,19 @@ struct FLiveLinkOpenTrackIOCache
 
 	/** List of transform subject names known. */
 	TSet<FName> TransformSubjectNames;
-	
+
 	/** Live Link Static Data, used to detect changes */
 	FLiveLinkStaticDataStruct StaticData;
 
 	/** OpenTrackIO static camera data (make, model, etc.) */
-	TOptional<FLiveLinkOpenTrackIOStaticCamera> StaticCamera;
+	FLiveLinkOpenTrackIOStaticCamera StaticCamera;
 
 	/** OpenTrackIO static lens data */
 	TOptional<FLiveLinkOpenTrackIOStaticLens> StaticLens;
 
 	/** Cached state of exposing the transforms as subjects. */
 	bool bSendTransformsAsSubjects = false;
-	
+
 	/**
 	 * Check the given sequence number and return true if it is the expected number. Sequence numbers should be incrementing
 	 * modulo UINT16_MAX.  This also accounts for delays in transmission based on sample rate provided by the first inbound packets.
@@ -47,13 +47,13 @@ struct FLiveLinkOpenTrackIOCache
 		}
 		LastDataReceiveTimeInSeconds = FPlatformTime::Seconds();
 	}
-	
+
 	/** Constructs a subject name from the given data. */
 	FName GetSubjectNameFromData(const FString& Prefix, const FLiveLinkOpenTrackIOData& Data);
 
 	/** Constructs a subject name from the given data. */
 	FName GetTransformName(const FLiveLinkOpenTrackIOTransform& InTransform);
-		
+
 	/** Create Live Link static data from the OpenTrackIO data */
 	FLiveLinkStaticDataStruct MakeStaticData(const FLiveLinkOpenTrackIOData& StaticData, const bool bApplyXform);
 
