Always set the camera static data. This will allow users to provide the sensor information but not require setting string values on the camera static data.

#Synced-CL 51338485

diff --git a/Engine/Plugins/VirtualProduction/LiveLinkOpenTrackIO/Source/LiveLinkOpenTrackIO/Private/LiveLinkOpenTrackIOSource.cpp b/Engine/Plugins/VirtualProduction/LiveLinkOpenTrackIO/Source/LiveLinkOpenTrackIO/Private/LiveLinkOpenTrackIOSource.cpp
--- a/Engine/Plugins/VirtualProduction/LiveLinkOpenTrackIO/Source/LiveLinkOpenTrackIO/Private/LiveLinkOpenTrackIOSource.cpp
+++ b/Engine/Plugins/VirtualProduction/LiveLinkOpenTrackIO/Source/LiveLinkOpenTrackIO/Private/LiveLinkOpenTrackIOSource.cpp
@@ -6,19 +6,19 @@
 
 #include "LiveLinkOpenTrackIOConversions.h"
 
-#include "LiveLinkTypes.h"
 #include "ILiveLinkClient.h"
+#include "LiveLinkTypes.h"
 
-#include "LiveLinkOpenTrackIOTypes.h"
 #include "LiveLinkOpenTrackIOLiveLinkTypes.h"
+#include "LiveLinkOpenTrackIOTypes.h"
 
 #include "LiveLinkOpenTrackIO.h"
-#include "LiveLinkOpenTrackIOSourceSettings.h"
 #include "LiveLinkOpenTrackIOParser.h"
+#include "LiveLinkOpenTrackIOSourceSettings.h"
 #include "LiveLinkOpenTrackIOTranscoder.h"
 
-#include "Roles/LiveLinkCameraRole.h"
 #include "LiveLinkOpenTrackIORole.h"
+#include "Roles/LiveLinkCameraRole.h"
 
 #include "HAL/Platform.h"
 #include "Misc/CoreDelegates.h"
@@ -26,12 +26,12 @@
 #include "Common/UdpSocketBuilder.h"
 #include "Common/UdpSocketReceiver.h"
 
+#include "Interfaces/IPv4/IPv4Address.h"
+#include "Interfaces/IPv4/IPv4Endpoint.h"
 #include "Roles/LiveLinkTransformRole.h"
 #include "Roles/LiveLinkTransformTypes.h"
-#include "Sockets.h"
 #include "SocketSubsystem.h"
-#include "Interfaces/IPv4/IPv4Address.h"
-#include "Interfaces/IPv4/IPv4Endpoint.h"
+#include "Sockets.h"
 
 #define LOCTEXT_NAMESPACE "LiveLinkOpenTrackIOSource"
 
@@ -68,13 +68,11 @@ void DoJoinMulticastGroup(const TSharedRef<FInternetAddr>& MulticastAddr, const
 	const bool bJoinedGroup = MulticastSocket->JoinMulticastGroup(*MulticastAddr, *IpAddr);
 	if (bJoinedGroup)
 	{
-		UE_LOG(LogLiveLinkOpenTrackIO, Display, TEXT("Added local interface '%s' to multicast group '%s'"),
-				*IpAddr->ToString(false), *MulticastAddr->ToString(true));
+		UE_LOG(LogLiveLinkOpenTrackIO, Display, TEXT("Added local interface '%s' to multicast group '%s'"), *IpAddr->ToString(false), *MulticastAddr->ToString(true));
 	}
 	else
 	{
-		UE_LOG(LogLiveLinkOpenTrackIO, Warning, TEXT("Failed to join multicast group '%s' on detected local interface '%s'"),
-				*MulticastAddr->ToString(true), *IpAddr->ToString(false));
+		UE_LOG(LogLiveLinkOpenTrackIO, Warning, TEXT("Failed to join multicast group '%s' on detected local interface '%s'"), *MulticastAddr->ToString(true), *IpAddr->ToString(false));
 	}
 }
 
@@ -101,12 +99,12 @@ void JoinedToGroup(const FIPv4Endpoint& UnicastEndpoint, const FIPv4Endpoint& Mu
 	else
 	{
 		TSharedRef<FInternetAddr> UnicastAddr = UnicastEndpoint.ToInternetAddr();
-		DoJoinMulticastGroup(MulticastAddr, UnicastAddr , MulticastSocket);
+		DoJoinMulticastGroup(MulticastAddr, UnicastAddr, MulticastSocket);
 	}
 #endif
 }
 
-}
+} // namespace UE::LiveLinkOpenTrackIO::Private
 
 FLiveLinkOpenTrackIOSource::FLiveLinkOpenTrackIOSource(FLiveLinkOpenTrackIOConnectionSettings InConnectionSettings)
 	: ConnectionSettings(MoveTemp(InConnectionSettings))
@@ -158,7 +156,7 @@ bool FLiveLinkOpenTrackIOSource::RequestSourceShutdown()
 	}
 
 	// We do a latent shutdown because of a possible deadlock due to a mutex shared between:
-	// 
+	//
 	// * PushSubjectFrameData_AnyThread and
 	// * RequestSourceShutdown
 	//
@@ -199,7 +197,7 @@ void FLiveLinkOpenTrackIOSource::Update()
 			SourceStatus = LOCTEXT("NotStarted", "Not started");
 
 			// Note: InitializeSettings should get us out of this state directly to ResetRequested
-		break;
+			break;
 
 		case ELiveLinkOpenTrackIOState::ResetRequested:
 
@@ -284,24 +282,24 @@ bool FLiveLinkOpenTrackIOSource::ParseEndpointFromSourceSettings()
 
 	switch (SavedSourceSettings->Protocol)
 	{
-	case ELiveLinkOpenTrackIONetworkProtocol::Unicast:
+		case ELiveLinkOpenTrackIONetworkProtocol::Unicast:
 
-		if (!UniEndpoint.IsSet())
-		{
-			return false;
-		}
-		break;
+			if (!UniEndpoint.IsSet())
+			{
+				return false;
+			}
+			break;
 
-	case ELiveLinkOpenTrackIONetworkProtocol::Multicast:
+		case ELiveLinkOpenTrackIONetworkProtocol::Multicast:
 
-		if (!UniEndpoint.IsSet() || !MultiEndpoint.IsSet() || !MultiEndpoint->Address.IsMulticastAddress())
-		{
-			return false;
-		}
-		break;
+			if (!UniEndpoint.IsSet() || !MultiEndpoint.IsSet() || !MultiEndpoint->Address.IsMulticastAddress())
+			{
+				return false;
+			}
+			break;
 
-	default:
-		return false;
+		default:
+			return false;
 	}
 
 	return true;
@@ -369,7 +367,7 @@ void FLiveLinkOpenTrackIOSource::HandleInboundData(const TSharedPtr<FArrayReader
 		{
 			LastSender_RunnableThread = Message.Sender;
 			AsyncTask(ENamedThreads::GameThread, [this, Sender = Message.Sender]() {
-					SourceMachineName = Sender.Address.ToText();
+				SourceMachineName = Sender.Address.ToText();
 			});
 		}
 
@@ -399,7 +397,7 @@ void FLiveLinkOpenTrackIOSource::HandleInboundData(const TSharedPtr<FArrayReader
 
 void FLiveLinkOpenTrackIOSource::RemoveAllTransformSubjects(FLiveLinkOpenTrackIOCache* Cache)
 {
-	// Remove any transform subjects based on this name change. 
+	// Remove any transform subjects based on this name change.
 	for (FName TransformSubjectName : Cache->TransformSubjectNames)
 	{
 		Client->RemoveSubject_AnyThread({ SourceGuid, TransformSubjectName });
@@ -409,7 +407,7 @@ void FLiveLinkOpenTrackIOSource::RemoveAllTransformSubjects(FLiveLinkOpenTrackIO
 
 void FLiveLinkOpenTrackIOSource::ConditionallyPushLiveLinkTransformData(FLiveLinkOpenTrackIOCache* Cache, const FLiveLinkOpenTrackIOData& InData)
 {
-	const bool bEnableTransformSubjects = IsValid(SavedSourceSettings) ? SavedSourceSettings->ShouldExtractTransformSubjects() : false; 
+	const bool bEnableTransformSubjects = IsValid(SavedSourceSettings) ? SavedSourceSettings->ShouldExtractTransformSubjects() : false;
 	if (!bEnableTransformSubjects)
 	{
 		RemoveAllTransformSubjects(Cache);
@@ -426,7 +424,7 @@ void FLiveLinkOpenTrackIOSource::ConditionallyPushLiveLinkTransformData(FLiveLin
 			NewStaticData.bIsLocationSupported = true;
 			NewStaticData.bIsRotationSupported = true;
 			NewStaticData.bIsScaleSupported = true;
-			
+
 			Client->PushSubjectStaticData_AnyThread({ SourceGuid, TransformName }, ULiveLinkTransformRole::StaticClass(), MoveTemp(StaticData));
 		}
 
@@ -435,7 +433,7 @@ void FLiveLinkOpenTrackIOSource::ConditionallyPushLiveLinkTransformData(FLiveLin
 		FLiveLinkTransformFrameData& NewFrameData = *FrameData.Cast<FLiveLinkTransformFrameData>();
 
 		NewFrameData.Transform = LiveLinkOpenTrackIOConversions::ToUnrealTransform(Transform);
-		Client->PushSubjectFrameData_AnyThread({SourceGuid,TransformName}, MoveTemp(FrameData));
+		Client->PushSubjectFrameData_AnyThread({ SourceGuid, TransformName }, MoveTemp(FrameData));
 	}
 }
 
@@ -464,12 +462,9 @@ void FLiveLinkOpenTrackIOSource::PushDataToLiveLink_AnyThread(const FLiveLinkOpe
 		UE_LOG(LogLiveLinkOpenTrackIO, Warning, TEXT("Received packet for %s is out of sequence. Discarding %d."), *InData.SourceId, Header.SequenceNumber);
 		return;
 	}
-	
+
 	// Update the OpenTrackIO static camera data (make, model, etc.)
-	if (InData.Static.Camera.IsValid())
-	{
-		Cache->StaticCamera = InData.Static.Camera;
-	}
+	Cache->StaticCamera = InData.Static.Camera;
 
 	// Update the OpenTrackIO static lens data
 	if (InData.Static.Lens.IsValid())
@@ -499,7 +494,7 @@ void FLiveLinkOpenTrackIOSource::PushDataToLiveLink_AnyThread(const FLiveLinkOpe
 	}
 
 	const bool bShouldApplyTransform = IsValid(SavedSourceSettings) && SavedSourceSettings->ShouldApplyXformToCamera();
-	
+
 	// Update the static data if it has changed or if the subject name was updated.
 	{
 		FLiveLinkStaticDataStruct StaticData = Cache->MakeStaticData(InData, bShouldApplyTransform);
@@ -511,7 +506,7 @@ void FLiveLinkOpenTrackIOSource::PushDataToLiveLink_AnyThread(const FLiveLinkOpe
 			Client->PushSubjectStaticData_AnyThread({ SourceGuid, Cache->SubjectName }, ULiveLinkOpenTrackIORole::StaticClass(), MoveTemp(StaticData));
 		}
 	}
-	
+
 	ConditionallyPushLiveLinkTransformData(Cache, InData);
 
 	// Push the per-frame data
@@ -520,13 +515,13 @@ void FLiveLinkOpenTrackIOSource::PushDataToLiveLink_AnyThread(const FLiveLinkOpe
 
 	// Note we are moving the InData into the LL data.
 	OpenTrackIOData.OpenTrackData = MoveTemp(InData);
-	
-	Client->PushSubjectFrameData_AnyThread(	{ SourceGuid, Cache->SubjectName },	MoveTemp(FrameData) );
-	
+
+	Client->PushSubjectFrameData_AnyThread({ SourceGuid, Cache->SubjectName }, MoveTemp(FrameData));
+
 	Cache->UpdateLastKnownSequenceNumber(Header.SequenceNumber);
 }
 
-void FLiveLinkOpenTrackIOSource::StopUdpReceivers() 
+void FLiveLinkOpenTrackIOSource::StopUdpReceivers()
 {
 	if (MulticastReceiver)
 	{
@@ -585,26 +580,26 @@ bool FLiveLinkOpenTrackIOSource::OpenMulticastSocket()
 
 	/** Create a multicast socket. */
 	MulticastSocket = FUdpSocketBuilder(TEXT("UdpOpenTrackIO_MulticastSocket"))
-		.AsNonBlocking()
-		.AsReusable()
+						  .AsNonBlocking()
+						  .AsReusable()
 #if PLATFORM_WINDOWS
 
-		/**
-		 * If multiple bus instances bind the same unicast ip:port combination (allowed as the socket is marked as reusable), then for each packet
-		 * sent to that ip:port combination, only one of the instances (at the discretion of the OS) will receive it. The instance that receives the
-		 * packet may vary over time, seemingly based on the congestion of its socket. This isn't the intended usage.
-		 *
-		 * To allow traffic to be sent directly to unicast for discovery, set the interface and port for the unicast endpoint
-		 * However for legacy reason, keep binding this as well, although it might be unreliable in some cases
-		 */
-		.BoundToAddress(UnicastEndpoint.Address)
+						  /**
+						   * If multiple bus instances bind the same unicast ip:port combination (allowed as the socket is marked as reusable), then for each packet
+						   * sent to that ip:port combination, only one of the instances (at the discretion of the OS) will receive it. The instance that receives the
+						   * packet may vary over time, seemingly based on the congestion of its socket. This isn't the intended usage.
+						   *
+						   * To allow traffic to be sent directly to unicast for discovery, set the interface and port for the unicast endpoint
+						   * However for legacy reason, keep binding this as well, although it might be unreliable in some cases
+						   */
+						  .BoundToAddress(UnicastEndpoint.Address)
 #endif
-		.BoundToPort(MulticastEndpoint.Port)
+						  .BoundToPort(MulticastEndpoint.Port)
 #if PLATFORM_SUPPORTS_UDP_MULTICAST_GROUP
-		.WithMulticastLoopback()
-		.WithMulticastInterface(UnicastEndpoint.Address)
+						  .WithMulticastLoopback()
+						  .WithMulticastInterface(UnicastEndpoint.Address)
 #endif
-		.WithReceiveBufferSize(ReceiveBufferSize);
+						  .WithReceiveBufferSize(ReceiveBufferSize);
 
 	if (!MulticastSocket)
 	{
@@ -633,10 +628,10 @@ bool FLiveLinkOpenTrackIOSource::OpenUnicastSocket()
 	const uint32 ReceiveBufferSize = 2 * 1024 * 1024;
 
 	UnicastSocket = FUdpSocketBuilder(TEXT("UdpOpenTrackIO_UnicastSocket"))
-		.AsNonBlocking()
-		.BoundToEndpoint(UnicastEndpoint)
-		.WithMulticastLoopback()
-		.WithReceiveBufferSize(ReceiveBufferSize);
+						.AsNonBlocking()
+						.BoundToEndpoint(UnicastEndpoint)
+						.WithMulticastLoopback()
+						.WithReceiveBufferSize(ReceiveBufferSize);
 
 	if (!UnicastSocket)
 	{
@@ -661,35 +656,34 @@ bool FLiveLinkOpenTrackIOSource::OpenUnicastSocket()
 	return true;
 }
 
-
 bool FLiveLinkOpenTrackIOSource::OpenSockets()
 {
 	switch (SavedSourceSettings->Protocol)
 	{
-	case ELiveLinkOpenTrackIONetworkProtocol::Unicast:
-	{
-		if (!OpenUnicastSocket())
+		case ELiveLinkOpenTrackIONetworkProtocol::Unicast:
 		{
-			return false;
+			if (!OpenUnicastSocket())
+			{
+				return false;
+			}
+
+			break;
 		}
+		case ELiveLinkOpenTrackIONetworkProtocol::Multicast:
+		{
+			if (!OpenMulticastSocket())
+			{
+				return false;
+			}
 
-		break;
-	}
-	case ELiveLinkOpenTrackIONetworkProtocol::Multicast:
-	{
-		if (!OpenMulticastSocket())
+			break;
+		}
+		default:
 		{
 			return false;
 		}
-
-		break;
 	}
-	default:
-	{
-		return false;
-	}
-	}
-	
+
 	return true;
 }
 
