Remove the old tests that relied on hardcoded values that are no longer relevant.    We have the python testing harness to validate crc and end-to-end testing.

#Synced-CL 51338485

diff --git a/Engine/Plugins/VirtualProduction/LiveLinkOpenTrackIO/Source/LiveLinkOpenTrackIO/Private/Tests/OpenTrackReaderTests.spec.cpp b/Engine/Plugins/VirtualProduction/LiveLinkOpenTrackIO/Source/LiveLinkOpenTrackIO/Private/Tests/OpenTrackReaderTests.spec.cpp
--- a/Engine/Plugins/VirtualProduction/LiveLinkOpenTrackIO/Source/LiveLinkOpenTrackIO/Private/Tests/OpenTrackReaderTests.spec.cpp
+++ b/Engine/Plugins/VirtualProduction/LiveLinkOpenTrackIO/Source/LiveLinkOpenTrackIO/Private/Tests/OpenTrackReaderTests.spec.cpp
@@ -1,5 +1,6 @@
 // Copyright Epic Games, Inc. All Rights Reserved.
 
+// clang-format off
 #include "OpenTrackIOTestHelpers.h"
 
 #if WITH_EDITOR
@@ -9,55 +10,22 @@
 
 #include "LiveLinkOpenTrackIOParser.h"
 #include "LiveLinkOpenTrackIOTypes.h"
-	#include "LiveLinkOpenTrackIODatagram.h"
 
 #if WITH_DEV_AUTOMATION_TESTS
 
+// clang-format on
 namespace UE::OpenTrackIO::Tests
 {
 
 BEGIN_DEFINE_SPEC(FOpenTrackIOReaderTests,
 	TEXT("Editor.LiveLinkOpenTrackIO.ReaderTests"),
 	EAutomationTestFlags::EditorContext | EAutomationTestFlags::EngineFilter)
-	const TArray<FString> Tests = {
-		TEXT("FullStaticOpenTrack"),
-		TEXT("RecommendedDynamicExample"),
-		TEXT("CompleteStaticExample")
-	};
-
-	struct FOTIO_DataValidation
-	{
-		uint16 SequenceNumber;
-		uint16 Checksum;
-	};
-
-	const TArray<FOTIO_DataValidation> CborValidationData = []() -> TArray<FOTIO_DataValidation> {
-		return {
-		{ .SequenceNumber = 2, .Checksum = 0x86C }
-		,{ .SequenceNumber = 3, .Checksum = 0xA587 }
-		,{ .SequenceNumber = 4, .Checksum = 0xF596 }
-		,{ .SequenceNumber = 5, .Checksum = 0x46A5 }
-		,{ .SequenceNumber = 6, .Checksum = 0x96B4 }
-		,{ .SequenceNumber = 7, .Checksum = 0xE6C3 }
-		,{ .SequenceNumber = 8, .Checksum = 0x12CC }
-		,{ .SequenceNumber = 9, .Checksum = 0x3DD5 }
-		,{ .SequenceNumber = 10, .Checksum = 0x68DE }
-		,{ .SequenceNumber = 11, .Checksum = 0x93E7 }
-		,{ .SequenceNumber = 12, .Checksum = 0xBEF0 }};
-	}();
-	const TArray<FOTIO_DataValidation> JsonValidationData = []() -> TArray<FOTIO_DataValidation> {
-		return {
-		{ .SequenceNumber = 2, .Checksum = 0xDA9C }
-		,{ .SequenceNumber = 3, .Checksum = 0xA293 }
-		,{ .SequenceNumber = 4, .Checksum = 0x7CA5 }
-		,{ .SequenceNumber = 5, .Checksum = 0x449C }
-		,{ .SequenceNumber = 6, .Checksum = 0x1EAE }
-		,{ .SequenceNumber = 7, .Checksum = 0xE5A5 }
-		,{ .SequenceNumber = 8, .Checksum = 0xBFB7 }
-		,{ .SequenceNumber = 9, .Checksum = 0x87AE }
-		,{ .SequenceNumber = 10, .Checksum = 0x61C0 }
-		,{ .SequenceNumber = 11, .Checksum = 0xA608 }};
-	}();
+const TArray<FString> Tests = {
+	TEXT("FullStaticOpenTrack"),
+	TEXT("RecommendedDynamicExample"),
+	TEXT("CompleteStaticExample")
+};
+
 END_DEFINE_SPEC(FOpenTrackIOReaderTests)
 
 void FOpenTrackIOReaderTests::Define()
@@ -65,10 +33,10 @@ void FOpenTrackIOReaderTests::Define()
 	using namespace UE::OpenTrackIO::Tests;
 
 	BeforeEach([this] {
-		// Nothing right now. 
+		// Nothing right now.
 	});
 	Describe("CanonicalCases", [this] {
-		It("Parses JSON", [this] {	
+		It("Parses JSON", [this] {
 			for (const FString& TestName : Tests)
 			{
 				FString JsonTestName = TestName + TEXT(".json");
@@ -83,16 +51,14 @@ void FOpenTrackIOReaderTests::Define()
 			}
 		});
 
-		It("Fails with invalid JSON",[this]
-		{
+		It("Fails with invalid JSON", [this] {
 			FString InvalidJson = TEXT("{\"not opentrack\" : {}}");
 			TOptional<FLiveLinkOpenTrackIOData> Data = UE::OpenTrackIO::Private::ParseJsonBlob(InvalidJson);
-			
+
 			TestTrue("Should have failed with bogus data", !Data.IsSet());
 		});
-		
-		It("Parses CBOR", [this]
-		{	
+
+		It("Parses CBOR", [this] {
 			for (const FString& TestName : Tests)
 			{
 				FString CborTestName = TestName + TEXT(".cbor");
@@ -107,88 +73,10 @@ void FOpenTrackIOReaderTests::Define()
 			}
 		});
 	});
-
-	Describe("ReadPackets", [this]
-	{
-		It("Reads CBOR Packets", [this] {
-			const FString CborPackets = "otio_cbor.packets";
-			FString FullPath = GetSampleFile(CborPackets);
-
-			TArray<uint8> PacketsBlob;
-			TestTrue("Did load cbor packets blob.", FFileHelper::LoadFileToArray(PacketsBlob, *FullPath));
-
-			TArrayView<uint8> AllPacketsView(PacketsBlob);
-			uint64 Index = 0;
-			for (const FOTIO_DataValidation& Item : CborValidationData)
-			{
-				TArrayView<const uint8> PacketView = AllPacketsView.Slice(Index, AllPacketsView.Num()-Index);
-
-				FOpenTrackIOHeaderWithPayload PayloadContainer;
-				const bool bPayloadIsGood =
-					UE::OpenTrackIO::Private::GetHeaderAndPayloadFromBytes(PacketView, PayloadContainer);
-
-				FString ParsedPacket = TEXT("Cbor Parsed Packet No ") + FString::FromInt(Item.SequenceNumber);
-				TestTrue(ParsedPacket, bPayloadIsGood);
-
-				if (TOptional<FLiveLinkOpenTrackIOData> ParsedPayload = UE::OpenTrackIO::Private::ParsePayload(PayloadContainer))
-				{
-					const FLiveLinkOpenTrackIODatagramHeader& Header = PayloadContainer.GetHeader();
-					const FLiveLinkOpenTrackIOData& Data = *ParsedPayload;
-
-					TestTrue(ParsedPacket + " Has Payload", Header.GetPayloadSize() > 0);
-					TestTrue(ParsedPacket + " Matches Checksum", Item.Checksum == Header.Checksum);
-					TestTrue(ParsedPacket + " Matches Sequence", Item.SequenceNumber == Header.SequenceNumber);
-					Index += Header.GetPayloadSize() + sizeof(FLiveLinkOpenTrackIODatagramHeader); 
-				}
-				else
-				{
-					break;
-				}
-				
-			}
-		});
-		It("Reads JSON Packets", [this] {
-			const FString JsonPackets = "otio_json.packets";
-			FString FullPath = GetSampleFile(JsonPackets);
-
-			TArray<uint8> PacketsBlob;
-			TestTrue("Did load packets blob.", FFileHelper::LoadFileToArray(PacketsBlob, *FullPath));
-
-			TArrayView<uint8> AllPacketsView(PacketsBlob);
-			uint64 Index = 0;
-			for (const FOTIO_DataValidation& Item : JsonValidationData)
-			{
-				TArrayView<const uint8> PacketView = AllPacketsView.Slice(Index, AllPacketsView.Num()-Index);
-
-				FOpenTrackIOHeaderWithPayload PayloadContainer;
-				const bool bPayloadIsGood =
-					UE::OpenTrackIO::Private::GetHeaderAndPayloadFromBytes(PacketView, PayloadContainer);
-
-				FString ParsedPacket = TEXT("Json Parsed Packet No ") + FString::FromInt(Item.SequenceNumber);
-				TestTrue(ParsedPacket, bPayloadIsGood);
-
-				if (TOptional<FLiveLinkOpenTrackIOData> ParsedPayload = UE::OpenTrackIO::Private::ParsePayload(PayloadContainer))
-				{
-					const FLiveLinkOpenTrackIODatagramHeader& Header = PayloadContainer.GetHeader();
-					const FLiveLinkOpenTrackIOData& Data = *ParsedPayload;
-
-					TestTrue(ParsedPacket + " Has Payload", Header.GetPayloadSize() > 0);
-					TestTrue(ParsedPacket + " Matches Checksum", Item.Checksum == Header.Checksum);
-					TestTrue(ParsedPacket + " Matches Sequence", Item.SequenceNumber == Header.SequenceNumber);
-					Index += Header.GetPayloadSize() + sizeof(FLiveLinkOpenTrackIODatagramHeader); 
-				}
-				else
-				{
-					break;
-				}
-				
-			}
-		});
-	});
-}
-	
 }
 
-#endif // WITH_DEV_AUTOMATION_TESTS
+} // namespace UE::OpenTrackIO::Tests
+
+	#endif // WITH_DEV_AUTOMATION_TESTS
 
 #endif // WITH_EDITOR
